Renamed HierarchyEditor -> DataHierarchyEditor to be a bit less generic
Moved most of (Data) Hierarchy Editor from Public headers to Internal headers so only engine code can use it for now.
When Data Hierarchy Editor has established itself we will change it to Public once more.

diff --git a/Engine/Config/BaseEngine.ini b/Engine/Config/BaseEngine.ini
--- a/Engine/Config/BaseEngine.ini
+++ b/Engine/Config/BaseEngine.ini
@@ -1344,6 +1344,15 @@ ClassName=/Script/OnlineSubsystemUtils.OnlineEngineInterfaceImpl
 +FunctionRedirects=(OldName="ChaosClothComponent.SetClothAsset", NewName="SetAsset")
 +PropertyRedirects=(OldName="ChaosClothComponent.ClothAsset", NewName="ChaosClothComponent.Asset")
 
+; 5.6
++PackageRedirects=(OldName="/Script/HierarchyEditor", NewName="/Script/DataHierarchyEditor")
++StructRedirects=(OldName="/Script/HierarchyEditor.HierarchyElementIdentity",NewName="/Script/DataHierarchyEditor.HierarchyElementIdentity")
++ClassRedirects=(OldName="/Script/HierarchyEditor.HierarchyElement",NewName="/Script/DataHierarchyEditor.HierarchyElement")
++ClassRedirects=(OldName="/Script/HierarchyEditor.HierarchyRoot",NewName="/Script/DataHierarchyEditor.HierarchyRoot")
++ClassRedirects=(OldName="/Script/HierarchyEditor.HierarchyItem",NewName="/Script/DataHierarchyEditor.HierarchyItem")
++ClassRedirects=(OldName="/Script/HierarchyEditor.HierarchySection",NewName="/Script/DataHierarchyEditor.HierarchySection")
++ClassRedirects=(OldName="/Script/HierarchyEditor.HierarchyCategory",NewName="/Script/DataHierarchyEditor.HierarchyCategory")
+
 ; End of releases redirects
 
 ; Fixing the output of tags for blueprints to be reset.
diff --git a/Engine/Plugins/FX/Niagara/Config/BaseNiagara.ini b/Engine/Plugins/FX/Niagara/Config/BaseNiagara.ini
--- a/Engine/Plugins/FX/Niagara/Config/BaseNiagara.ini
+++ b/Engine/Plugins/FX/Niagara/Config/BaseNiagara.ini
@@ -51,13 +51,13 @@
 +PropertyRedirects=(OldName="NiagaraSystem.FlipbookGeneratedSettings", NewName="NiagaraSystem.BakerGeneratedSettings")
 +StructRedirects=(OldName="NiagaraFlipbookTextureSettings", NewName="/Script/Niagara.NiagaraBakerTextureSettings")
 +StructRedirects=(OldName="NiagaraFlipbookTextureSource", NewName="/Script/Niagara.NiagaraBakerTextureSource")
-; Hierarchy Editor Redirects
-+StructRedirects=(OldName="/Script/NiagaraEditor.NiagaraHierarchyIdentity",NewName="/Script/HierarchyEditor.HierarchyElementIdentity")
-+ClassRedirects=(OldName="/Script/NiagaraEditor.NiagaraHierarchyItemBase",NewName="/Script/HierarchyEditor.HierarchyElement")
-+ClassRedirects=(OldName="/Script/NiagaraEditor.NiagaraHierarchyRoot",NewName="/Script/HierarchyEditor.HierarchyRoot")
-+ClassRedirects=(OldName="/Script/NiagaraEditor.NiagaraHierarchyItem",NewName="/Script/HierarchyEditor.HierarchyItem")
-+ClassRedirects=(OldName="/Script/NiagaraEditor.NiagaraHierarchySection",NewName="/Script/HierarchyEditor.HierarchySection")
-+ClassRedirects=(OldName="/Script/NiagaraEditor.NiagaraHierarchyCategory",NewName="/Script/HierarchyEditor.HierarchyCategory")
+; Data Hierarchy Editor Redirects
++StructRedirects=(OldName="/Script/NiagaraEditor.NiagaraHierarchyIdentity",NewName="/Script/DataHierarchyEditor.HierarchyElementIdentity")
++ClassRedirects=(OldName="/Script/NiagaraEditor.NiagaraHierarchyItemBase",NewName="/Script/DataHierarchyEditor.HierarchyElement")
++ClassRedirects=(OldName="/Script/NiagaraEditor.NiagaraHierarchyRoot",NewName="/Script/DataHierarchyEditor.HierarchyRoot")
++ClassRedirects=(OldName="/Script/NiagaraEditor.NiagaraHierarchyItem",NewName="/Script/DataHierarchyEditor.HierarchyItem")
++ClassRedirects=(OldName="/Script/NiagaraEditor.NiagaraHierarchySection",NewName="/Script/DataHierarchyEditor.HierarchySection")
++ClassRedirects=(OldName="/Script/NiagaraEditor.NiagaraHierarchyCategory",NewName="/Script/DataHierarchyEditor.HierarchyCategory")
 
 [/Script/Niagara.NiagaraBakerSettings]
 BakeQualityLevel=Cinematic
diff --git a/Engine/Plugins/FX/Niagara/Source/NiagaraEditor/NiagaraEditor.Build.cs b/Engine/Plugins/FX/Niagara/Source/NiagaraEditor/NiagaraEditor.Build.cs
--- a/Engine/Plugins/FX/Niagara/Source/NiagaraEditor/NiagaraEditor.Build.cs
+++ b/Engine/Plugins/FX/Niagara/Source/NiagaraEditor/NiagaraEditor.Build.cs
@@ -89,7 +89,7 @@ public class NiagaraEditor : ModuleRules
                 "AppFramework",
 				"ToolWidgets",
 				"Sequencer",
-				"HierarchyEditor"
+				"DataHierarchyEditor"
             }
         );
 
diff --git a/Engine/Plugins/FX/Niagara/Source/NiagaraEditor/Private/Customizations/NiagaraComponentDetails.h b/Engine/Plugins/FX/Niagara/Source/NiagaraEditor/Private/Customizations/NiagaraComponentDetails.h
--- a/Engine/Plugins/FX/Niagara/Source/NiagaraEditor/Private/Customizations/NiagaraComponentDetails.h
+++ b/Engine/Plugins/FX/Niagara/Source/NiagaraEditor/Private/Customizations/NiagaraComponentDetails.h
@@ -8,7 +8,6 @@
 #include "NiagaraComponent.h"
 #include "NiagaraSimCacheCapture.h"
 #include "ViewModels/TNiagaraViewModelManager.h"
-#include "HierarchyViewModelBase.h"
 #include "NiagaraComponentDetails.generated.h"
 
 class UNiagaraSystem;
diff --git a/Engine/Plugins/FX/Niagara/Source/NiagaraEditor/Private/Customizations/UserParameters/NiagaraComponentUserParametersBuilder.cpp b/Engine/Plugins/FX/Niagara/Source/NiagaraEditor/Private/Customizations/UserParameters/NiagaraComponentUserParametersBuilder.cpp
--- a/Engine/Plugins/FX/Niagara/Source/NiagaraEditor/Private/Customizations/UserParameters/NiagaraComponentUserParametersBuilder.cpp
+++ b/Engine/Plugins/FX/Niagara/Source/NiagaraEditor/Private/Customizations/UserParameters/NiagaraComponentUserParametersBuilder.cpp
@@ -501,9 +501,9 @@ void FNiagaraComponentUserParametersNodeBuilder::RegisterRebuildOnHierarchyChang
 			if(SystemViewModel.IsValid())
 			{
 				SystemViewModel->GetUserParametersHierarchyViewModel()->OnHierarchyChanged().RemoveAll(this);
-				SystemViewModel->GetUserParametersHierarchyViewModel()->OnHierarchyChanged().Add(UHierarchyViewModelBase::FOnHierarchyChanged::FDelegate::CreateSP(this, &FNiagaraComponentUserParametersNodeBuilder::Rebuild));
+				SystemViewModel->GetUserParametersHierarchyViewModel()->OnHierarchyChanged().Add(UDataHierarchyViewModelBase::FOnHierarchyChanged::FDelegate::CreateSP(this, &FNiagaraComponentUserParametersNodeBuilder::Rebuild));
 				SystemViewModel->GetUserParametersHierarchyViewModel()->OnHierarchyPropertiesChanged().RemoveAll(this);
-				SystemViewModel->GetUserParametersHierarchyViewModel()->OnHierarchyPropertiesChanged().Add(UHierarchyViewModelBase::FOnHierarchyPropertiesChanged::FDelegate::CreateSP(this, &FNiagaraComponentUserParametersNodeBuilder::Rebuild));
+				SystemViewModel->GetUserParametersHierarchyViewModel()->OnHierarchyPropertiesChanged().Add(UDataHierarchyViewModelBase::FOnHierarchyPropertiesChanged::FDelegate::CreateSP(this, &FNiagaraComponentUserParametersNodeBuilder::Rebuild));
 			}
 		}
 	}
diff --git a/Engine/Plugins/FX/Niagara/Source/NiagaraEditor/Private/Customizations/UserParameters/NiagaraSystemUserParametersBuilder.cpp b/Engine/Plugins/FX/Niagara/Source/NiagaraEditor/Private/Customizations/UserParameters/NiagaraSystemUserParametersBuilder.cpp
--- a/Engine/Plugins/FX/Niagara/Source/NiagaraEditor/Private/Customizations/UserParameters/NiagaraSystemUserParametersBuilder.cpp
+++ b/Engine/Plugins/FX/Niagara/Source/NiagaraEditor/Private/Customizations/UserParameters/NiagaraSystemUserParametersBuilder.cpp
@@ -81,9 +81,9 @@ void FNiagaraSystemUserParameterBuilder::GenerateChildContent(IDetailChildrenBui
 			System->GetExposedParameters().OnStructureChanged().Add(FNiagaraParameterStore::FOnStructureChanged::FDelegate::CreateSP(this, &FNiagaraSystemUserParameterBuilder::Rebuild));
 
 			SystemViewModel.Pin()->GetUserParametersHierarchyViewModel()->OnHierarchyChanged().RemoveAll(this);
-			SystemViewModel.Pin()->GetUserParametersHierarchyViewModel()->OnHierarchyChanged().Add(UHierarchyViewModelBase::FOnHierarchyChanged::FDelegate::CreateSP(this, &FNiagaraSystemUserParameterBuilder::Rebuild));
+			SystemViewModel.Pin()->GetUserParametersHierarchyViewModel()->OnHierarchyChanged().Add(UDataHierarchyViewModelBase::FOnHierarchyChanged::FDelegate::CreateSP(this, &FNiagaraSystemUserParameterBuilder::Rebuild));
 			SystemViewModel.Pin()->GetUserParametersHierarchyViewModel()->OnHierarchyPropertiesChanged().RemoveAll(this);
-			SystemViewModel.Pin()->GetUserParametersHierarchyViewModel()->OnHierarchyPropertiesChanged().Add(UHierarchyViewModelBase::FOnHierarchyPropertiesChanged::FDelegate::CreateSP(this, &FNiagaraSystemUserParameterBuilder::Rebuild));
+			SystemViewModel.Pin()->GetUserParametersHierarchyViewModel()->OnHierarchyPropertiesChanged().Add(UDataHierarchyViewModelBase::FOnHierarchyPropertiesChanged::FDelegate::CreateSP(this, &FNiagaraSystemUserParameterBuilder::Rebuild));
 			
 			TSharedPtr<FNiagaraUserParameterPanelViewModel> UserParameterPanelViewModel = SystemViewModel.Pin()->GetUserParameterPanelViewModel();
 			if(ensure(UserParameterPanelViewModel.IsValid()))
diff --git a/Engine/Plugins/FX/Niagara/Source/NiagaraEditor/Private/Customizations/UserParameters/NiagaraSystemUserParametersBuilder.h b/Engine/Plugins/FX/Niagara/Source/NiagaraEditor/Private/Customizations/UserParameters/NiagaraSystemUserParametersBuilder.h
--- a/Engine/Plugins/FX/Niagara/Source/NiagaraEditor/Private/Customizations/UserParameters/NiagaraSystemUserParametersBuilder.h
+++ b/Engine/Plugins/FX/Niagara/Source/NiagaraEditor/Private/Customizations/UserParameters/NiagaraSystemUserParametersBuilder.h
@@ -6,7 +6,7 @@
 #include "CoreMinimal.h"
 #include "Input/Reply.h"
 #include "ViewModels/TNiagaraViewModelManager.h"
-#include "HierarchyViewModelBase.h"
+#include "DataHierarchyViewModelBase.h"
 #include "ViewModels/HierarchyEditor/NiagaraUserParametersHierarchyViewModel.h"
 
 class IDetailLayoutBuilder;
diff --git a/Engine/Plugins/FX/Niagara/Source/NiagaraEditor/Private/Customizations/UserParameters/NiagaraUserParametersBuilderBase.h b/Engine/Plugins/FX/Niagara/Source/NiagaraEditor/Private/Customizations/UserParameters/NiagaraUserParametersBuilderBase.h
--- a/Engine/Plugins/FX/Niagara/Source/NiagaraEditor/Private/Customizations/UserParameters/NiagaraUserParametersBuilderBase.h
+++ b/Engine/Plugins/FX/Niagara/Source/NiagaraEditor/Private/Customizations/UserParameters/NiagaraUserParametersBuilderBase.h
@@ -6,7 +6,7 @@
 #include "IDetailCustomNodeBuilder.h"
 #include "NiagaraComponent.h"
 #include "SNiagaraParameterEditor.h"
-#include "HierarchyViewModelBase.h"
+#include "DataHierarchyViewModelBase.h"
 
 class UNiagaraScriptVariable;
 
diff --git a/Engine/Plugins/FX/Niagara/Source/NiagaraEditor/Private/NiagaraEditorUtilities.cpp b/Engine/Plugins/FX/Niagara/Source/NiagaraEditor/Private/NiagaraEditorUtilities.cpp
--- a/Engine/Plugins/FX/Niagara/Source/NiagaraEditor/Private/NiagaraEditorUtilities.cpp
+++ b/Engine/Plugins/FX/Niagara/Source/NiagaraEditor/Private/NiagaraEditorUtilities.cpp
@@ -80,7 +80,7 @@
 #include "ViewModels/NiagaraEmitterHandleViewModel.h"
 #include "ViewModels/NiagaraEmitterViewModel.h"
 #include "ViewModels/NiagaraParameterPanelViewModel.h"
-#include "Widgets/SHierarchyEditor.h"
+#include "Widgets/SDataHierarchyEditor.h"
 #include "ViewModels/HierarchyEditor/NiagaraHierarchyScriptParametersViewModel.h"
 #include "Widgets/SToolTip.h"
 
diff --git a/Engine/Plugins/FX/Niagara/Source/NiagaraEditor/Private/NiagaraGraph.cpp b/Engine/Plugins/FX/Niagara/Source/NiagaraEditor/Private/NiagaraGraph.cpp
--- a/Engine/Plugins/FX/Niagara/Source/NiagaraEditor/Private/NiagaraGraph.cpp
+++ b/Engine/Plugins/FX/Niagara/Source/NiagaraEditor/Private/NiagaraGraph.cpp
@@ -37,6 +37,7 @@
 #include "Modules/ModuleManager.h"
 #include "String/ParseTokens.h"
 #include "ViewModels/NiagaraScriptViewModel.h"
+#include "DataHierarchyViewModelBase.h"
 #include "ViewModels/HierarchyEditor/NiagaraHierarchyScriptParametersViewModel.h"
 
 #include UE_INLINE_GENERATED_CPP_BY_NAME(NiagaraGraph)
diff --git a/Engine/Plugins/FX/Niagara/Source/NiagaraEditor/Private/Toolkits/NiagaraScriptToolkit.cpp b/Engine/Plugins/FX/Niagara/Source/NiagaraEditor/Private/Toolkits/NiagaraScriptToolkit.cpp
--- a/Engine/Plugins/FX/Niagara/Source/NiagaraEditor/Private/Toolkits/NiagaraScriptToolkit.cpp
+++ b/Engine/Plugins/FX/Niagara/Source/NiagaraEditor/Private/Toolkits/NiagaraScriptToolkit.cpp
@@ -41,7 +41,7 @@
 #include "ViewModels/NiagaraParameterPanelViewModel.h"
 #include "ViewModels/NiagaraScriptViewModel.h"
 #include "ViewModels/HierarchyEditor/NiagaraHierarchyScriptParametersViewModel.h"
-#include "Widgets/SHierarchyEditor.h"
+#include "Widgets/SDataHierarchyEditor.h"
 #include "Widgets/Docking/SDockTab.h"
 #include "Widgets/SNiagaraParameterDefinitionsPanel.h"
 #include "Widgets/SNiagaraParameterPanel.h"
@@ -635,7 +635,7 @@ TSharedRef<SDockTab> FNiagaraScriptToolkit::SpawnTabHierarchyEditor(const FSpawn
 			SNew(SBox)
 			.AddMetaData<FTagMetaData>(FTagMetaData(TEXT("UserParameters")))
 			[
-				SNew(SHierarchyEditor, ParametersHierarchyViewModel)
+				SNew(SDataHierarchyEditor, ParametersHierarchyViewModel)
 				.OnGenerateRowContentWidget_Static(&FNiagaraEditorUtilities::HierarchyEditor::Scripts::GenerateRowContentForScriptParameterHierarchyEditor)
 				//.OnGenerateCustomDetailsPanelNameWidget_Static(&GenerateCustomDetailsPanelNameWidgetForUserParameterEditor)
 			]
diff --git a/Engine/Plugins/FX/Niagara/Source/NiagaraEditor/Private/Toolkits/SystemToolkitModes/NiagaraSystemToolkitModeBase.cpp b/Engine/Plugins/FX/Niagara/Source/NiagaraEditor/Private/Toolkits/SystemToolkitModes/NiagaraSystemToolkitModeBase.cpp
--- a/Engine/Plugins/FX/Niagara/Source/NiagaraEditor/Private/Toolkits/SystemToolkitModes/NiagaraSystemToolkitModeBase.cpp
+++ b/Engine/Plugins/FX/Niagara/Source/NiagaraEditor/Private/Toolkits/SystemToolkitModes/NiagaraSystemToolkitModeBase.cpp
@@ -27,7 +27,7 @@
 #include "Widgets/SNiagaraScriptGraph.h"
 #include "Widgets/SNiagaraSelectedObjectsDetails.h"
 #include "Widgets/SNiagaraGeneratedCodeView.h"
-#include "Widgets/SHierarchyEditor.h"
+#include "Widgets/SDataHierarchyEditor.h"
 #include "Widgets/Layout/SWidgetSwitcher.h"
 #include "Widgets/Layout/SBox.h"
 #include "Widgets/Input/SButton.h"
@@ -913,7 +913,7 @@ TSharedRef<SWidget> GenerateCustomDetailsPanelNameWidgetForUserParameterEditor(T
 			{
 				return FText::FromString(HierarchyItem->ToString());
 			})
-			.TextStyle(&FHierarchyEditorStyle::Get().GetWidgetStyle<FTextBlockStyle>("HierarchyEditor.CategoryTextBlock"))
+			.TextStyle(&FDataHierarchyEditorStyle::Get().GetWidgetStyle<FTextBlockStyle>("HierarchyEditor.CategoryTextBlock"))
 		];
 	}
 	else if(const UNiagaraHierarchyUserParameter* UserParameter = Cast<UNiagaraHierarchyUserParameter>(HierarchyItem->GetData()))
@@ -934,7 +934,7 @@ TSharedRef<SDockTab> FNiagaraSystemToolkitModeBase::SpawnTab_UserParametersHiera
 			SNew(SBox)
 			.AddMetaData<FTagMetaData>(FTagMetaData(TEXT("UserParameters")))
 			[
-				SNew(SHierarchyEditor, SystemToolkit.Pin()->GetSystemViewModel()->GetUserParametersHierarchyViewModel())
+				SNew(SDataHierarchyEditor, SystemToolkit.Pin()->GetSystemViewModel()->GetUserParametersHierarchyViewModel())
 				.OnGenerateRowContentWidget_Static(&GenerateRowContentForUserParameterHierarchyEditor, SystemToolkit.Pin()->GetSystemViewModel())
 				.OnGenerateCustomDetailsPanelNameWidget_Static(&GenerateCustomDetailsPanelNameWidgetForUserParameterEditor)
 			]
@@ -1082,7 +1082,7 @@ TSharedRef<SWidget> FNiagaraSystemToolkitModeBase::CreateSummaryViewWidget() con
 		TSharedPtr<FNiagaraEmitterViewModel> EmitterViewModel = SystemToolkit.Pin()->GetSystemViewModel()->GetEmitterHandleViewModels()[0]->GetEmitterViewModel();
 		UNiagaraSummaryViewViewModel* SummaryViewHierarchyViewModel = EmitterViewModel->GetSummaryHierarchyViewModel();
 		
-		ContentWidget = SNew(SHierarchyEditor, SummaryViewHierarchyViewModel)
+		ContentWidget = SNew(SDataHierarchyEditor, SummaryViewHierarchyViewModel)
 		.OnGenerateRowContentWidget_Static(&GenerateRowContentForSummaryViewHierarchyEditor, EmitterViewModel)
 		.OnGenerateCustomDetailsPanelNameWidget_Static(&GenerateCustomDetailsPanelNameWidgetForSummaryViewEditor);
 	}
@@ -1102,7 +1102,7 @@ TSharedRef<SWidget> FNiagaraSystemToolkitModeBase::CreateSummaryViewWidget() con
 			TSharedPtr<FNiagaraEmitterViewModel> EmitterViewModel = SystemToolkit.Pin()->GetSystemViewModel()->GetEmitterHandleViewModelById(SelectedEmitterHandleGuids[0])->GetEmitterViewModel();
 			UNiagaraSummaryViewViewModel* SummaryViewHierarchyViewModel = EmitterViewModel->GetSummaryHierarchyViewModel();
 		
-			ContentWidget = SNew(SHierarchyEditor, SummaryViewHierarchyViewModel)
+			ContentWidget = SNew(SDataHierarchyEditor, SummaryViewHierarchyViewModel)
 			.OnGenerateRowContentWidget_Static(&GenerateRowContentForSummaryViewHierarchyEditor, EmitterViewModel)
 			.OnGenerateCustomDetailsPanelNameWidget_Static(&GenerateCustomDetailsPanelNameWidgetForSummaryViewEditor);
 		}
@@ -1141,7 +1141,7 @@ TSharedRef<SWidget> FNiagaraSystemToolkitModeBase::CreateScratchPadHierarchyWidg
 	{
 		UNiagaraHierarchyScriptParametersViewModel* ScriptHierarchyViewModel = ActiveScratchPadViewModel->GetHierarchyViewModel();
 		
-		ContentWidget = SNew(SHierarchyEditor, ScriptHierarchyViewModel)
+		ContentWidget = SNew(SDataHierarchyEditor, ScriptHierarchyViewModel)
 		.OnGenerateRowContentWidget_Static(&FNiagaraEditorUtilities::HierarchyEditor::Scripts::GenerateRowContentForScriptParameterHierarchyEditor);
 		//.OnGenerateCustomDetailsPanelNameWidget_Static(&GenerateCustomDetailsPanelNameWidgetForSummaryViewEditor);
 	}
diff --git a/Engine/Plugins/FX/Niagara/Source/NiagaraEditor/Private/ViewModels/HierarchyEditor/NiagaraHierarchyScriptParametersViewModel.cpp b/Engine/Plugins/FX/Niagara/Source/NiagaraEditor/Private/ViewModels/HierarchyEditor/NiagaraHierarchyScriptParametersViewModel.cpp
--- a/Engine/Plugins/FX/Niagara/Source/NiagaraEditor/Private/ViewModels/HierarchyEditor/NiagaraHierarchyScriptParametersViewModel.cpp
+++ b/Engine/Plugins/FX/Niagara/Source/NiagaraEditor/Private/ViewModels/HierarchyEditor/NiagaraHierarchyScriptParametersViewModel.cpp
@@ -76,7 +76,7 @@ void UNiagaraHierarchyScriptParametersViewModel::Initialize(TSharedRef<FNiagaraS
 	ScriptViewModelWeak = InScriptViewModel;
 
 	Cast<UNiagaraScriptSource>(ScriptViewModelWeak.Pin()->GetStandaloneScript().GetScriptData()->GetSource())->NodeGraph->OnParametersChanged().AddUObject(this, &UNiagaraHierarchyScriptParametersViewModel::OnParametersChanged);
-	UHierarchyViewModelBase::Initialize();
+	UDataHierarchyViewModelBase::Initialize();
 
 	UNiagaraHierarchyScriptParameterRefreshContext* ScriptParameterRefreshContext = NewObject<UNiagaraHierarchyScriptParameterRefreshContext>(this);
 	ScriptParameterRefreshContext->SetNiagaraGraph(Cast<UNiagaraScriptSource>(InScriptViewModel->GetStandaloneScript().GetScriptData()->GetSource())->NodeGraph);
diff --git a/Engine/Plugins/FX/Niagara/Source/NiagaraEditor/Private/ViewModels/HierarchyEditor/NiagaraSummaryViewViewModel.cpp b/Engine/Plugins/FX/Niagara/Source/NiagaraEditor/Private/ViewModels/HierarchyEditor/NiagaraSummaryViewViewModel.cpp
--- a/Engine/Plugins/FX/Niagara/Source/NiagaraEditor/Private/ViewModels/HierarchyEditor/NiagaraSummaryViewViewModel.cpp
+++ b/Engine/Plugins/FX/Niagara/Source/NiagaraEditor/Private/ViewModels/HierarchyEditor/NiagaraSummaryViewViewModel.cpp
@@ -128,7 +128,7 @@ void UNiagaraSummaryViewViewModel::Initialize(TSharedRef<FNiagaraEmitterViewMode
 		EmitterViewModel->GetEmitter().Emitter->OnSimStagesChanged().AddUObject(this, &UNiagaraSummaryViewViewModel::OnSimStagesChanged);
 		EmitterViewModel->GetEmitter().Emitter->OnEventHandlersChanged().AddUObject(this, &UNiagaraSummaryViewViewModel::OnEventHandlersChanged);
 
-		UHierarchyViewModelBase::Initialize();
+		UDataHierarchyViewModelBase::Initialize();
 	}
 }
 
@@ -153,7 +153,7 @@ TSharedRef<FNiagaraEmitterViewModel> UNiagaraSummaryViewViewModel::GetEmitterVie
 
 TWeakObjectPtr<UNiagaraNodeFunctionCall> FNiagaraFunctionViewModel::GetFunctionCallNode() const
 {
-	TWeakObjectPtr<UHierarchyViewModelBase> ViewModel = GetHierarchyViewModel();
+	TWeakObjectPtr<UDataHierarchyViewModelBase> ViewModel = GetHierarchyViewModel();
 	UNiagaraSummaryViewViewModel* SummaryViewModel = CastChecked<UNiagaraSummaryViewViewModel>(ViewModel.Get());
 	return SummaryViewModel->GetFunctionCallNode(GetData()->GetPersistentIdentity().Guids[0]);
 }
@@ -169,7 +169,7 @@ void FNiagaraFunctionViewModel::OnScriptApplied(UNiagaraScript* NiagaraScript, F
 
 void FNiagaraFunctionViewModel::ClearCache() const
 {
-	TWeakObjectPtr<UHierarchyViewModelBase> ViewModel = GetHierarchyViewModel();
+	TWeakObjectPtr<UDataHierarchyViewModelBase> ViewModel = GetHierarchyViewModel();
 	UNiagaraSummaryViewViewModel* SummaryViewModel = CastChecked<UNiagaraSummaryViewViewModel>(ViewModel.Get());
 	SummaryViewModel->ClearFunctionCallNodeCache(GetData()->GetPersistentIdentity().Guids[0]);
 }
@@ -471,7 +471,7 @@ void FNiagaraModuleInputViewModel::AppendDynamicContextMenuForSingleElement(UToo
 
 TWeakObjectPtr<UNiagaraNodeFunctionCall> FNiagaraModuleInputViewModel::GetModuleNode() const
 {
-	TWeakObjectPtr<UHierarchyViewModelBase> ViewModel = GetHierarchyViewModel();
+	TWeakObjectPtr<UDataHierarchyViewModelBase> ViewModel = GetHierarchyViewModel();
 	UNiagaraSummaryViewViewModel* SummaryViewModel = CastChecked<UNiagaraSummaryViewViewModel>(ViewModel.Get());
 	return SummaryViewModel->GetFunctionCallNode(GetData()->GetPersistentIdentity().Guids[0]);
 }
@@ -598,7 +598,7 @@ void FNiagaraModuleInputViewModel::RefreshChildrenDataInternal()
 
 TOptional<FInputData> FNiagaraModuleInputViewModel::FindInputDataInternal() const
 {
-	TWeakObjectPtr<UHierarchyViewModelBase> ViewModel = GetHierarchyViewModel();
+	TWeakObjectPtr<UDataHierarchyViewModelBase> ViewModel = GetHierarchyViewModel();
 	UNiagaraSummaryViewViewModel* SummaryViewModel = CastChecked<UNiagaraSummaryViewViewModel>(ViewModel.Get());
 	return SummaryViewModel->GetInputData(*GetData<UNiagaraHierarchyModuleInput>());
 }
@@ -724,7 +724,7 @@ FHierarchyElementViewModel::FCanPerformActionResults FNiagaraAssignmentInputView
 
 TWeakObjectPtr<UNiagaraNodeAssignment> FNiagaraAssignmentInputViewModel::GetAssignmentNode() const
 {
-	TWeakObjectPtr<UHierarchyViewModelBase> ViewModel = GetHierarchyViewModel();
+	TWeakObjectPtr<UDataHierarchyViewModelBase> ViewModel = GetHierarchyViewModel();
 	UNiagaraSummaryViewViewModel* SummaryViewModel = CastChecked<UNiagaraSummaryViewViewModel>(ViewModel.Get());
 	return Cast<UNiagaraNodeAssignment>(SummaryViewModel->GetFunctionCallNode(GetData()->GetPersistentIdentity().Guids[0]));
 }
@@ -771,7 +771,7 @@ bool FNiagaraAssignmentInputViewModel::IsFromBaseEmitter() const
 void FNiagaraAssignmentInputViewModel::ClearCache() const
 {
 	InputDataCache.Reset();
-	TWeakObjectPtr<UHierarchyViewModelBase> ViewModel = GetHierarchyViewModel();
+	TWeakObjectPtr<UDataHierarchyViewModelBase> ViewModel = GetHierarchyViewModel();
 	UNiagaraSummaryViewViewModel* SummaryViewModel = CastChecked<UNiagaraSummaryViewViewModel>(ViewModel.Get());
 	SummaryViewModel->ClearFunctionCallNodeCache(GetData()->GetPersistentIdentity().Guids[0]);
 }
@@ -795,7 +795,7 @@ FHierarchyElementViewModel::FCanPerformActionResults FNiagaraAssignmentInputView
 
 TOptional<FNiagaraStackGraphUtilities::FMatchingFunctionInputData> FNiagaraAssignmentInputViewModel::FindInputDataInternal() const
 {
-	TWeakObjectPtr<UHierarchyViewModelBase> ViewModel = GetHierarchyViewModel();
+	TWeakObjectPtr<UDataHierarchyViewModelBase> ViewModel = GetHierarchyViewModel();
 	UNiagaraSummaryViewViewModel* SummaryViewModel = CastChecked<UNiagaraSummaryViewViewModel>(ViewModel.Get());
 	
 	if(GetAssignmentNode().IsValid())
@@ -956,7 +956,7 @@ FString FNiagaraHierarchyEventHandlerViewModel::ToString() const
 
 FNiagaraEventScriptProperties* FNiagaraHierarchyEventHandlerViewModel::GetEventScriptProperties() const
 {
-	TWeakObjectPtr<UHierarchyViewModelBase> ViewModel = GetHierarchyViewModel();
+	TWeakObjectPtr<UDataHierarchyViewModelBase> ViewModel = GetHierarchyViewModel();
 	UNiagaraSummaryViewViewModel* SummaryViewModel = CastChecked<UNiagaraSummaryViewViewModel>(ViewModel.Get());
 	FGuid UsageID = GetData()->GetPersistentIdentity().Guids[0];
 	
@@ -1057,7 +1057,7 @@ const UHierarchySection* FNiagaraHierarchyEventHandlerViewModel::GetSectionInter
 
 UNiagaraRendererProperties* FNiagaraHierarchyRendererViewModel::GetRendererProperties() const
 {
-	TWeakObjectPtr<UHierarchyViewModelBase> ViewModel = GetHierarchyViewModel();
+	TWeakObjectPtr<UDataHierarchyViewModelBase> ViewModel = GetHierarchyViewModel();
 	UNiagaraSummaryViewViewModel* SummaryViewModel = CastChecked<UNiagaraSummaryViewViewModel>(ViewModel.Get());
 	FGuid RendererGuid = GetData()->GetPersistentIdentity().Guids[0];
 
@@ -1082,7 +1082,7 @@ FString FNiagaraHierarchyEventHandlerPropertiesViewModel::ToString() const
 
 FNiagaraEventScriptProperties* FNiagaraHierarchyEventHandlerPropertiesViewModel::GetEventScriptProperties() const
 {
-	TWeakObjectPtr<UHierarchyViewModelBase> ViewModel = GetHierarchyViewModel();
+	TWeakObjectPtr<UDataHierarchyViewModelBase> ViewModel = GetHierarchyViewModel();
 	UNiagaraSummaryViewViewModel* SummaryViewModel = CastChecked<UNiagaraSummaryViewViewModel>(ViewModel.Get());
 	FGuid UsageID = GetData()->GetPersistentIdentity().Guids[0];
 	
@@ -1165,7 +1165,7 @@ FString FNiagaraHierarchySimStageViewModel::ToString() const
 
 UNiagaraSimulationStageBase* FNiagaraHierarchySimStageViewModel::GetSimStage() const
 {
-	TWeakObjectPtr<UHierarchyViewModelBase> ViewModel = GetHierarchyViewModel();
+	TWeakObjectPtr<UDataHierarchyViewModelBase> ViewModel = GetHierarchyViewModel();
 	UNiagaraSummaryViewViewModel* SummaryViewModel = CastChecked<UNiagaraSummaryViewViewModel>(ViewModel.Get());
 	FGuid SimStageGuid = GetData()->GetPersistentIdentity().Guids[0];
 
@@ -1272,7 +1272,7 @@ FString FNiagaraHierarchySimStagePropertiesViewModel::ToString() const
 
 UNiagaraSimulationStageBase* FNiagaraHierarchySimStagePropertiesViewModel::GetSimStage() const
 {
-	TWeakObjectPtr<UHierarchyViewModelBase> ViewModel = GetHierarchyViewModel();
+	TWeakObjectPtr<UDataHierarchyViewModelBase> ViewModel = GetHierarchyViewModel();
 	UNiagaraSummaryViewViewModel* SummaryViewModel = CastChecked<UNiagaraSummaryViewViewModel>(ViewModel.Get());
 	FGuid SimStageGuid = GetData()->GetPersistentIdentity().Guids[0];
 
diff --git a/Engine/Plugins/FX/Niagara/Source/NiagaraEditor/Private/ViewModels/HierarchyEditor/NiagaraUserParametersHierarchyViewModel.cpp b/Engine/Plugins/FX/Niagara/Source/NiagaraEditor/Private/ViewModels/HierarchyEditor/NiagaraUserParametersHierarchyViewModel.cpp
--- a/Engine/Plugins/FX/Niagara/Source/NiagaraEditor/Private/ViewModels/HierarchyEditor/NiagaraUserParametersHierarchyViewModel.cpp
+++ b/Engine/Plugins/FX/Niagara/Source/NiagaraEditor/Private/ViewModels/HierarchyEditor/NiagaraUserParametersHierarchyViewModel.cpp
@@ -46,7 +46,7 @@ TSharedRef<FNiagaraSystemViewModel> UNiagaraUserParametersHierarchyViewModel::Ge
 void UNiagaraUserParametersHierarchyViewModel::Initialize(TSharedRef<FNiagaraSystemViewModel> InSystemViewModel)
 {
 	SystemViewModelWeak = InSystemViewModel;
-	UHierarchyViewModelBase::Initialize();
+	UDataHierarchyViewModelBase::Initialize();
 
 	UNiagaraHierarchyUserParameterRefreshContext* UserParameterRefreshContext = NewObject<UNiagaraHierarchyUserParameterRefreshContext>(this);
 	UserParameterRefreshContext->SetSystem(&InSystemViewModel->GetSystem());
diff --git a/Engine/Plugins/FX/Niagara/Source/NiagaraEditor/Private/ViewModels/Stack/NiagaraStackEmitterSettingsGroup.cpp b/Engine/Plugins/FX/Niagara/Source/NiagaraEditor/Private/ViewModels/Stack/NiagaraStackEmitterSettingsGroup.cpp
--- a/Engine/Plugins/FX/Niagara/Source/NiagaraEditor/Private/ViewModels/Stack/NiagaraStackEmitterSettingsGroup.cpp
+++ b/Engine/Plugins/FX/Niagara/Source/NiagaraEditor/Private/ViewModels/Stack/NiagaraStackEmitterSettingsGroup.cpp
@@ -22,7 +22,6 @@
 #include "NiagaraNodeFunctionCall.h"
 #include "NiagaraSettings.h"
 #include "Toolkits/SystemToolkitModes/NiagaraSystemToolkitModeBase.h"
-#include "Widgets/SHierarchyEditor.h"
 
 #include UE_INLINE_GENERATED_CPP_BY_NAME(NiagaraStackEmitterSettingsGroup)
 
diff --git a/Engine/Plugins/FX/Niagara/Source/NiagaraEditor/Private/ViewModels/Stack/NiagaraStackEntry.cpp b/Engine/Plugins/FX/Niagara/Source/NiagaraEditor/Private/ViewModels/Stack/NiagaraStackEntry.cpp
--- a/Engine/Plugins/FX/Niagara/Source/NiagaraEditor/Private/ViewModels/Stack/NiagaraStackEntry.cpp
+++ b/Engine/Plugins/FX/Niagara/Source/NiagaraEditor/Private/ViewModels/Stack/NiagaraStackEntry.cpp
@@ -12,7 +12,7 @@
 #include "NiagaraScriptMergeManager.h"
 #include "Misc/SecureHash.h"
 #include "ScopedTransaction.h"
-#include "HierarchyViewModelBase.h"
+#include "DataHierarchyViewModelBase.h"
 #include "ViewModels/HierarchyEditor/NiagaraSummaryViewViewModel.h"
 #include "ViewModels/Stack/NiagaraStackClipboardUtilities.h"
 
diff --git a/Engine/Plugins/FX/Niagara/Source/NiagaraEditor/Private/ViewModels/Stack/NiagaraStackGraphUtilities.cpp b/Engine/Plugins/FX/Niagara/Source/NiagaraEditor/Private/ViewModels/Stack/NiagaraStackGraphUtilities.cpp
--- a/Engine/Plugins/FX/Niagara/Source/NiagaraEditor/Private/ViewModels/Stack/NiagaraStackGraphUtilities.cpp
+++ b/Engine/Plugins/FX/Niagara/Source/NiagaraEditor/Private/ViewModels/Stack/NiagaraStackGraphUtilities.cpp
@@ -48,7 +48,7 @@
 #include "NiagaraSettings.h"
 #include "NiagaraSimulationStageBase.h"
 #include "NiagaraEmitter.h"
-#include "HierarchyViewModelBase.h"
+#include "DataHierarchyViewModelBase.h"
 #include "ViewModels/HierarchyEditor/NiagaraSummaryViewViewModel.h"
 #include "TraversalCache/StaticContextBuilder.h"
 #include "TraversalCache/TraversalCache.h"
diff --git a/Engine/Plugins/FX/Niagara/Source/NiagaraEditor/Private/ViewModels/Stack/NiagaraStackObject.cpp b/Engine/Plugins/FX/Niagara/Source/NiagaraEditor/Private/ViewModels/Stack/NiagaraStackObject.cpp
--- a/Engine/Plugins/FX/Niagara/Source/NiagaraEditor/Private/ViewModels/Stack/NiagaraStackObject.cpp
+++ b/Engine/Plugins/FX/Niagara/Source/NiagaraEditor/Private/ViewModels/Stack/NiagaraStackObject.cpp
@@ -15,7 +15,7 @@
 #include "IDetailTreeNode.h"
 
 #include "NiagaraPlatformSet.h"
-#include "HierarchyViewModelBase.h"
+#include "DataHierarchyViewModelBase.h"
 #include "ViewModels/Stack/NiagaraStackObjectIssueGenerator.h"
 
 #include UE_INLINE_GENERATED_CPP_BY_NAME(NiagaraStackObject)
diff --git a/Engine/Plugins/FX/Niagara/Source/NiagaraEditor/Private/ViewModels/Stack/NiagaraStackPropertyRow.cpp b/Engine/Plugins/FX/Niagara/Source/NiagaraEditor/Private/ViewModels/Stack/NiagaraStackPropertyRow.cpp
--- a/Engine/Plugins/FX/Niagara/Source/NiagaraEditor/Private/ViewModels/Stack/NiagaraStackPropertyRow.cpp
+++ b/Engine/Plugins/FX/Niagara/Source/NiagaraEditor/Private/ViewModels/Stack/NiagaraStackPropertyRow.cpp
@@ -8,7 +8,7 @@
 #include "PropertyHandle.h"
 #include "IDetailPropertyRow.h"
 #include "ScopedTransaction.h"
-#include "HierarchyViewModelBase.h"
+#include "DataHierarchyViewModelBase.h"
 
 #include UE_INLINE_GENERATED_CPP_BY_NAME(NiagaraStackPropertyRow)
 
diff --git a/Engine/Plugins/FX/Niagara/Source/NiagaraEditor/Private/Widgets/SNiagaraScriptInputPreviewPanel.cpp b/Engine/Plugins/FX/Niagara/Source/NiagaraEditor/Private/Widgets/SNiagaraScriptInputPreviewPanel.cpp
--- a/Engine/Plugins/FX/Niagara/Source/NiagaraEditor/Private/Widgets/SNiagaraScriptInputPreviewPanel.cpp
+++ b/Engine/Plugins/FX/Niagara/Source/NiagaraEditor/Private/Widgets/SNiagaraScriptInputPreviewPanel.cpp
@@ -2,8 +2,8 @@
 
 #include "Widgets/SNiagaraScriptInputPreviewPanel.h"
 
-#include "HierarchyEditorStyle.h"
-#include "HierarchyViewModelBase.h"
+#include "DataHierarchyEditorStyle.h"
+#include "DataHierarchyViewModelBase.h"
 #include "NiagaraEditorStyle.h"
 #include "NiagaraEditorUtilities.h"
 #include "NiagaraObjectSelection.h"
@@ -133,7 +133,7 @@ TSharedRef<ITableRow> SNiagaraScriptInputPreviewPanel::OnGenerateRow(UHierarchyE
 	{		
 		TableRow = SNew(STableRow<UHierarchyElement*>, TableViewBase)
 		.ToolTipText_UObject(Category, &UHierarchyCategory::GetTooltip)
-		.Style(FHierarchyEditorStyle::Get(), "HierarchyEditor.Row.Category")
+		.Style(FDataHierarchyEditorStyle::Get(), "HierarchyEditor.Row.Category")
 		.Padding(FMargin(0.f, 6.f))
 		[
 			SNew(SRichTextBlock)
diff --git a/Engine/Plugins/FX/Niagara/Source/NiagaraEditor/Public/NiagaraEditorUtilities.h b/Engine/Plugins/FX/Niagara/Source/NiagaraEditor/Public/NiagaraEditorUtilities.h
--- a/Engine/Plugins/FX/Niagara/Source/NiagaraEditor/Public/NiagaraEditorUtilities.h
+++ b/Engine/Plugins/FX/Niagara/Source/NiagaraEditor/Public/NiagaraEditorUtilities.h
@@ -12,7 +12,7 @@
 #include "NiagaraEditorSettings.h"
 #include "ViewModels/NiagaraSystemScalabilityViewModel.h"
 #include "ViewModels/NiagaraSystemViewModel.h"
-#include "HierarchyViewModelBase.h"
+#include "DataHierarchyViewModelBase.h"
 
 class UNiagaraClipboardContent;
 class UNiagaraNodeInput;
diff --git a/Engine/Plugins/FX/Niagara/Source/NiagaraEditor/Public/NiagaraEmitterEditorData.h b/Engine/Plugins/FX/Niagara/Source/NiagaraEditor/Public/NiagaraEmitterEditorData.h
--- a/Engine/Plugins/FX/Niagara/Source/NiagaraEditor/Public/NiagaraEmitterEditorData.h
+++ b/Engine/Plugins/FX/Niagara/Source/NiagaraEditor/Public/NiagaraEmitterEditorData.h
@@ -5,7 +5,7 @@
 #include "NiagaraEditorDataBase.h"
 #include "NiagaraEditorCommon.h"
 #include "NiagaraStackSection.h"
-#include "HierarchyViewModelBase.h"
+#include "DataHierarchyViewModelBase.h"
 #include "NiagaraEmitterEditorData.generated.h"
 
 class UNiagaraStackEditorData;
diff --git a/Engine/Plugins/FX/Niagara/Source/NiagaraEditor/Public/NiagaraSystemEditorData.h b/Engine/Plugins/FX/Niagara/Source/NiagaraEditor/Public/NiagaraSystemEditorData.h
--- a/Engine/Plugins/FX/Niagara/Source/NiagaraEditor/Public/NiagaraSystemEditorData.h
+++ b/Engine/Plugins/FX/Niagara/Source/NiagaraEditor/Public/NiagaraSystemEditorData.h
@@ -5,7 +5,7 @@
 #include "Misc/FrameRate.h"
 #include "NiagaraEditorDataBase.h"
 #include "NiagaraEditorData.h"
-#include "HierarchyViewModelBase.h"
+#include "DataHierarchyViewModelBase.h"
 
 #include "NiagaraSystemEditorData.generated.h"
 
diff --git a/Engine/Plugins/FX/Niagara/Source/NiagaraEditor/Public/ViewModels/HierarchyEditor/NiagaraHierarchyScriptParametersViewModel.h b/Engine/Plugins/FX/Niagara/Source/NiagaraEditor/Public/ViewModels/HierarchyEditor/NiagaraHierarchyScriptParametersViewModel.h
--- a/Engine/Plugins/FX/Niagara/Source/NiagaraEditor/Public/ViewModels/HierarchyEditor/NiagaraHierarchyScriptParametersViewModel.h
+++ b/Engine/Plugins/FX/Niagara/Source/NiagaraEditor/Public/ViewModels/HierarchyEditor/NiagaraHierarchyScriptParametersViewModel.h
@@ -3,7 +3,7 @@
 #pragma once
 
 #include "NiagaraGraph.h"
-#include "HierarchyViewModelBase.h"
+#include "DataHierarchyViewModelBase.h"
 #include "NiagaraScriptVariable.h"
 #include "ViewModels/NiagaraScriptViewModel.h"
 #include "NiagaraHierarchyScriptParametersViewModel.generated.h"
@@ -53,14 +53,14 @@ class UNiagaraHierarchyScriptCategory : public UHierarchyCategory
  * The view model that defines the script editor's hierarchy editor for input parameters.
  */
 UCLASS()
-class NIAGARAEDITOR_API UNiagaraHierarchyScriptParametersViewModel : public UHierarchyViewModelBase
+class NIAGARAEDITOR_API UNiagaraHierarchyScriptParametersViewModel : public UDataHierarchyViewModelBase
 {
 	GENERATED_BODY()
 public:
 	void Initialize(TSharedRef<FNiagaraScriptViewModel> InScriptViewModel);
 	TSharedRef<FNiagaraScriptViewModel> GetScriptViewModel() const;
 
-	/** UHierarchyViewModelBase */
+	/** UDataHierarchyViewModelBase */
 	virtual UHierarchyRoot* GetHierarchyRoot() const override;
 	/** We set the outer for the source root to the graph, so that both source & hierarchy items have the same outer. */
 	virtual UObject* GetOuterForSourceRoot() const override;
diff --git a/Engine/Plugins/FX/Niagara/Source/NiagaraEditor/Public/ViewModels/HierarchyEditor/NiagaraSummaryViewViewModel.h b/Engine/Plugins/FX/Niagara/Source/NiagaraEditor/Public/ViewModels/HierarchyEditor/NiagaraSummaryViewViewModel.h
--- a/Engine/Plugins/FX/Niagara/Source/NiagaraEditor/Public/ViewModels/HierarchyEditor/NiagaraSummaryViewViewModel.h
+++ b/Engine/Plugins/FX/Niagara/Source/NiagaraEditor/Public/ViewModels/HierarchyEditor/NiagaraSummaryViewViewModel.h
@@ -6,7 +6,7 @@
 #include "NiagaraNodeFunctionCall.h"
 #include "NiagaraNodeAssignment.h"
 #include "NiagaraSimulationStageBase.h"
-#include "HierarchyViewModelBase.h"
+#include "DataHierarchyViewModelBase.h"
 #include "NiagaraSummaryViewViewModel.generated.h"
 
 class UNiagaraStackEventHandlerPropertiesItem;
@@ -137,7 +137,7 @@ public:
 };
 
 UCLASS(MinimalAPI)
-class UNiagaraSummaryViewViewModel : public UHierarchyViewModelBase
+class UNiagaraSummaryViewViewModel : public UDataHierarchyViewModelBase
 {
 	GENERATED_BODY()
 public:
diff --git a/Engine/Plugins/FX/Niagara/Source/NiagaraEditor/Public/ViewModels/HierarchyEditor/NiagaraUserParametersHierarchyViewModel.h b/Engine/Plugins/FX/Niagara/Source/NiagaraEditor/Public/ViewModels/HierarchyEditor/NiagaraUserParametersHierarchyViewModel.h
--- a/Engine/Plugins/FX/Niagara/Source/NiagaraEditor/Public/ViewModels/HierarchyEditor/NiagaraUserParametersHierarchyViewModel.h
+++ b/Engine/Plugins/FX/Niagara/Source/NiagaraEditor/Public/ViewModels/HierarchyEditor/NiagaraUserParametersHierarchyViewModel.h
@@ -2,7 +2,7 @@
 
 #pragma once
 
-#include "HierarchyViewModelBase.h"
+#include "DataHierarchyViewModelBase.h"
 #include "NiagaraScriptVariable.h"
 #include "NiagaraUserParametersHierarchyViewModel.generated.h"
 
@@ -44,7 +44,7 @@ private:
 
 struct FNiagaraHierarchyUserParameterViewModel : public FHierarchyItemViewModel
 {
-	FNiagaraHierarchyUserParameterViewModel(UNiagaraHierarchyUserParameter* InItem, TSharedRef<FHierarchyElementViewModel> InParent, TWeakObjectPtr<UHierarchyViewModelBase> InHierarchyViewModel)
+	FNiagaraHierarchyUserParameterViewModel(UNiagaraHierarchyUserParameter* InItem, TSharedRef<FHierarchyElementViewModel> InParent, TWeakObjectPtr<UDataHierarchyViewModelBase> InHierarchyViewModel)
 	: FHierarchyItemViewModel(InItem, InParent, InHierarchyViewModel) {}
 	
 	/** For editing in the details panel we want to handle the script variable that is represented by the hierarchy item, not the hierarchy item itself. */
@@ -57,7 +57,7 @@ struct FNiagaraHierarchyUserParameterViewModel : public FHierarchyItemViewModel
 };
 
 UCLASS()
-class UNiagaraUserParametersHierarchyViewModel : public UHierarchyViewModelBase
+class UNiagaraUserParametersHierarchyViewModel : public UDataHierarchyViewModelBase
 {
 	GENERATED_BODY()
 public:
diff --git a/Engine/Plugins/FX/Niagara/Source/NiagaraEditor/Public/ViewModels/Stack/NiagaraStackInputCategory.h b/Engine/Plugins/FX/Niagara/Source/NiagaraEditor/Public/ViewModels/Stack/NiagaraStackInputCategory.h
--- a/Engine/Plugins/FX/Niagara/Source/NiagaraEditor/Public/ViewModels/Stack/NiagaraStackInputCategory.h
+++ b/Engine/Plugins/FX/Niagara/Source/NiagaraEditor/Public/ViewModels/Stack/NiagaraStackInputCategory.h
@@ -6,7 +6,7 @@
 #include "NiagaraStackScriptHierarchyRoot.h"
 #include "ViewModels/Stack/NiagaraStackItem.h"
 #include "NiagaraTypes.h"
-#include "HierarchyViewModelBase.h"
+#include "DataHierarchyViewModelBase.h"
 #include "ViewModels/HierarchyEditor/NiagaraSummaryViewViewModel.h"
 #include "NiagaraStackInputCategory.generated.h"
 
diff --git a/Engine/Plugins/FX/Niagara/Source/NiagaraEditor/Public/ViewModels/Stack/NiagaraStackSummaryViewInputCollection.h b/Engine/Plugins/FX/Niagara/Source/NiagaraEditor/Public/ViewModels/Stack/NiagaraStackSummaryViewInputCollection.h
--- a/Engine/Plugins/FX/Niagara/Source/NiagaraEditor/Public/ViewModels/Stack/NiagaraStackSummaryViewInputCollection.h
+++ b/Engine/Plugins/FX/Niagara/Source/NiagaraEditor/Public/ViewModels/Stack/NiagaraStackSummaryViewInputCollection.h
@@ -4,7 +4,7 @@
 
 #include "NiagaraEmitter.h"
 #include "NiagaraStackValueCollection.h"
-#include "HierarchyViewModelBase.h"
+#include "DataHierarchyViewModelBase.h"
 #include "NiagaraStackSummaryViewInputCollection.generated.h"
 
 UCLASS(MinimalAPI)
diff --git a/Engine/Plugins/FX/Niagara/Source/NiagaraEditor/Public/Widgets/SNiagaraScriptInputPreviewPanel.h b/Engine/Plugins/FX/Niagara/Source/NiagaraEditor/Public/Widgets/SNiagaraScriptInputPreviewPanel.h
--- a/Engine/Plugins/FX/Niagara/Source/NiagaraEditor/Public/Widgets/SNiagaraScriptInputPreviewPanel.h
+++ b/Engine/Plugins/FX/Niagara/Source/NiagaraEditor/Public/Widgets/SNiagaraScriptInputPreviewPanel.h
@@ -5,7 +5,7 @@
 #include "CoreMinimal.h"
 #include "NiagaraGraph.h"
 #include "NiagaraScriptVariable.h"
-#include "HierarchyViewModelBase.h"
+#include "DataHierarchyViewModelBase.h"
 #include "Widgets/SCompoundWidget.h"
 #include "Widgets/Input/SSearchBox.h"
 #include "Widgets/Views/STreeView.h"
diff --git a/Engine/Plugins/FX/Niagara/Source/NiagaraEditorWidgets/NiagaraEditorWidgets.Build.cs b/Engine/Plugins/FX/Niagara/Source/NiagaraEditorWidgets/NiagaraEditorWidgets.Build.cs
--- a/Engine/Plugins/FX/Niagara/Source/NiagaraEditorWidgets/NiagaraEditorWidgets.Build.cs
+++ b/Engine/Plugins/FX/Niagara/Source/NiagaraEditorWidgets/NiagaraEditorWidgets.Build.cs
@@ -36,7 +36,7 @@ public class NiagaraEditorWidgets : ModuleRules
 			"ToolWidgets",
 			"KismetWidgets",
 			"ToolMenus",
-			"HierarchyEditor"
+			"DataHierarchyEditor"
 		});
 
 		PrivateIncludePathModuleNames.AddRange(new string[] {
diff --git a/Engine/Source/Editor/DataHierarchyEditor/DataHierarchyEditor.Build.cs b/Engine/Source/Editor/DataHierarchyEditor/DataHierarchyEditor.Build.cs
new file mode 100644
--- /dev/null
+++ b/Engine/Source/Editor/DataHierarchyEditor/DataHierarchyEditor.Build.cs
@@ -0,0 +1,32 @@
+﻿// Copyright Epic Games, Inc. All Rights Reserved.
+
+using UnrealBuildTool;
+
+public class DataHierarchyEditor : ModuleRules
+{
+    public DataHierarchyEditor(ReadOnlyTargetRules Target) : base(Target)
+    {
+        PCHUsage = ModuleRules.PCHUsageMode.UseExplicitOrSharedPCHs;
+
+        PublicDependencyModuleNames.AddRange(
+            new string[]
+            {
+                "Core", 
+                "CoreUObject",
+                "EditorWidgets",
+                "SlateCore"
+            }
+        );
+
+        PrivateDependencyModuleNames.AddRange(
+            new string[]
+            {
+                "Engine",
+                "Slate",
+                "UnrealEd",
+                "ToolMenus",
+                "InputCore"
+            }
+        );
+    }
+}
\ No newline at end of file
diff --git a/Engine/Source/Editor/DataHierarchyEditor/Internal/DataHierarchyEditorCommands.h b/Engine/Source/Editor/DataHierarchyEditor/Internal/DataHierarchyEditorCommands.h
new file mode 100644
--- /dev/null
+++ b/Engine/Source/Editor/DataHierarchyEditor/Internal/DataHierarchyEditorCommands.h
@@ -0,0 +1,25 @@
+// Copyright Epic Games, Inc. All Rights Reserved.
+
+#pragma once
+
+#include "CoreMinimal.h"
+#include "Styling/CoreStyle.h"
+#include "Framework/Commands/Commands.h"
+
+class DATAHIERARCHYEDITOR_API FDataHierarchyEditorCommands : public TCommands<FDataHierarchyEditorCommands>
+{
+public:
+	
+	FDataHierarchyEditorCommands()
+		: TCommands<FDataHierarchyEditorCommands>( TEXT("DataHierarchyEditorCommands"), NSLOCTEXT("DataHierarchyEditorCommands", "Data Hierarchy Editor Commands", "Data Hierarchy Editor Commands"), NAME_None, "DataHierarchyEditorStyle" )
+	{
+	}
+
+	virtual ~FDataHierarchyEditorCommands() override
+	{
+	}
+
+	virtual void RegisterCommands() override;	
+
+	TSharedPtr<FUICommandInfo> FindInHierarchy;
+};
diff --git a/Engine/Source/Editor/DataHierarchyEditor/Internal/DataHierarchyEditorMisc.h b/Engine/Source/Editor/DataHierarchyEditor/Internal/DataHierarchyEditorMisc.h
new file mode 100644
--- /dev/null
+++ b/Engine/Source/Editor/DataHierarchyEditor/Internal/DataHierarchyEditorMisc.h
@@ -0,0 +1,12 @@
+﻿// Copyright Epic Games, Inc. All Rights Reserved.
+
+#pragma once
+
+#include "Containers/Set.h"
+#include "UObject/NameTypes.h"
+
+
+namespace UE::DataHierarchyEditor
+{
+	DATAHIERARCHYEDITOR_API FName GetUniqueName(FName CandidateName, const TSet<FName>& ExistingNames);
+}
diff --git a/Engine/Source/Editor/DataHierarchyEditor/Internal/DataHierarchyEditorModule.h b/Engine/Source/Editor/DataHierarchyEditor/Internal/DataHierarchyEditorModule.h
new file mode 100644
--- /dev/null
+++ b/Engine/Source/Editor/DataHierarchyEditor/Internal/DataHierarchyEditorModule.h
@@ -0,0 +1,15 @@
+﻿// Copyright Epic Games, Inc. All Rights Reserved.
+
+#pragma once
+
+#include "CoreMinimal.h"
+#include "Modules/ModuleManager.h"
+
+DECLARE_LOG_CATEGORY_EXTERN(LogDataHierarchyEditor, Log, All)
+
+class FDataHierarchyEditorModule : public IModuleInterface
+{
+public:
+    virtual void StartupModule() override;
+    virtual void ShutdownModule() override;
+};
diff --git a/Engine/Source/Editor/DataHierarchyEditor/Internal/DataHierarchyEditorStyle.h b/Engine/Source/Editor/DataHierarchyEditor/Internal/DataHierarchyEditorStyle.h
new file mode 100644
--- /dev/null
+++ b/Engine/Source/Editor/DataHierarchyEditor/Internal/DataHierarchyEditorStyle.h
@@ -0,0 +1,29 @@
+// Copyright Epic Games, Inc. All Rights Reserved.
+
+#pragma once
+
+#include "Styling/SlateStyle.h"
+
+/** Manages the style which provides resources for Hierarchy Editor widgets. */
+class DATAHIERARCHYEDITOR_API FDataHierarchyEditorStyle : public FSlateStyleSet
+{
+public:
+	static void Register();
+	static void Unregister();
+	static void Shutdown();
+
+	/** reloads textures used by slate renderer */
+	static void ReloadTextures();
+
+	/** @return The Slate style set for Hierarchy Editor widgets */
+	static const FDataHierarchyEditorStyle& Get();
+
+	static void ReinitializeStyle();
+
+private:	
+	FDataHierarchyEditorStyle();
+
+	void InitDataHierarchyEditor();
+
+	static TSharedPtr<FDataHierarchyEditorStyle> DataHierarchyEditorStyle;
+};
diff --git a/Engine/Source/Editor/DataHierarchyEditor/Internal/Widgets/SDataHierarchyEditor.h b/Engine/Source/Editor/DataHierarchyEditor/Internal/Widgets/SDataHierarchyEditor.h
new file mode 100644
--- /dev/null
+++ b/Engine/Source/Editor/DataHierarchyEditor/Internal/Widgets/SDataHierarchyEditor.h
@@ -0,0 +1,281 @@
+// Copyright Epic Games, Inc. All Rights Reserved.
+
+#pragma once
+
+#include "CoreMinimal.h"
+#include "SDropTarget.h"
+#include "IDetailsView.h"
+#include "DataHierarchyEditorStyle.h"
+#include "Widgets/Views/STreeView.h"
+#include "Widgets/Views/STableRow.h"
+#include "Widgets/Text/SInlineEditableTextBlock.h"
+#include "Widgets/Input/SMenuAnchor.h"
+#include "Styling/SlateTypes.h"
+#include "DataHierarchyViewModelBase.h"
+#include "Widgets/Input/SSearchBox.h"
+#include "Widgets/Input/SCheckBox.h"
+#include "Misc/NotifyHook.h"
+#include "UObject/GCObject.h"
+
+class UDataHierarchyViewModelBase;
+class UHierarchySection;
+struct FHierarchyElementViewModel;
+struct FHierarchySectionViewModel;
+
+class DATAHIERARCHYEDITOR_API SHierarchySection : public SCompoundWidget
+{
+	DECLARE_DELEGATE_OneParam(FOnSectionActivated, TSharedPtr<FHierarchySectionViewModel> SectionViewModel)
+	
+	SLATE_BEGIN_ARGS(SHierarchySection)
+		: _bForbidDropOn(false)
+		{}
+		SLATE_ATTRIBUTE(ECheckBoxState, IsSectionActive)
+		SLATE_EVENT(FOnSectionActivated, OnSectionActivated)
+		SLATE_ARGUMENT(bool, bForbidDropOn)
+	SLATE_END_ARGS()
+
+public:
+	void Construct(const FArguments& InArgs, TSharedPtr<struct FHierarchySectionViewModel> InSection, TWeakObjectPtr<UDataHierarchyViewModelBase> HierarchyViewModel);
+	virtual ~SHierarchySection() override;
+
+	bool OnCanAcceptDrop(TSharedPtr<FDragDropOperation> DragDropOperation, EItemDropZone ItemDropZone) const;
+	FReply OnDroppedOn(const FGeometry&, const FDragDropEvent& DragDropEvent, EItemDropZone DropZone) const;
+
+	void TryEnterEditingMode() const;
+
+	TSharedPtr<struct FHierarchySectionViewModel> GetSectionViewModel();
+private:
+	virtual FReply OnMouseButtonDown(const FGeometry& MyGeometry, const FPointerEvent& MouseEvent) override;
+	virtual FReply OnMouseButtonUp(const FGeometry& MyGeometry, const FPointerEvent& MouseEvent) override;
+	
+	virtual void OnDragEnter(const FGeometry& MyGeometry, const FDragDropEvent& DragDropEvent) override;
+	virtual FReply OnDragDetected(const FGeometry& MyGeometry, const FPointerEvent& MouseEvent) override;
+	virtual void OnDragLeave(const FDragDropEvent& DragDropEvent) override;
+
+	virtual FReply OnKeyDown(const FGeometry& MyGeometry, const FKeyEvent& InKeyEvent) override;
+	
+	virtual void OnMouseLeave(const FPointerEvent& MouseEvent) override;
+
+	UHierarchySection* TryGetSectionData() const;
+
+	FText GetText() const;
+	FText GetTooltipText() const;
+	void OnRenameSection(const FText& Text, ETextCommit::Type CommitType) const;
+	bool OnVerifySectionRename(const FText& NewName, FText& OutTooltip) const;
+
+	bool IsSectionSelected() const;
+	bool IsSectionReadOnly() const;
+	ECheckBoxState GetSectionCheckState() const;
+	void OnSectionCheckChanged(ECheckBoxState NewState);
+	EActiveTimerReturnType ActivateSectionIfDragging(double CurrentTime, float DeltaTime) const;
+private:
+	TSharedPtr<SMenuAnchor> MenuAnchor;
+	TSharedPtr<SCheckBox> CheckBox;
+	TSharedPtr<SInlineEditableTextBlock> InlineEditableTextBlock;
+	TWeakObjectPtr<UDataHierarchyViewModelBase> HierarchyViewModel;
+	TSharedPtr<FHierarchySectionViewModel> SectionViewModel;
+private:
+	TAttribute<ECheckBoxState> IsSectionActive;
+	FOnSectionActivated OnSectionActivatedDelegate;
+	bool bForbidDropOn = false;
+	mutable bool bDraggedOn = false;
+};
+
+class DATAHIERARCHYEDITOR_API SDataHierarchyEditor : public SCompoundWidget, public FNotifyHook
+{
+	DECLARE_DELEGATE_RetVal_OneParam(TSharedRef<SWidget>, FOnGenerateRowContentWidget, TSharedRef<FHierarchyElementViewModel> HierarchyElement);
+	DECLARE_DELEGATE_RetVal_OneParam(TSharedRef<SWidget>, FOnGenerateCustomDetailsPanelNameWidget, TSharedPtr<FHierarchyElementViewModel> HierarchyElement);
+
+	SLATE_BEGIN_ARGS(SDataHierarchyEditor)
+		: _bReadOnly(true)
+		, _ItemRowStyle(&FAppStyle::Get().GetWidgetStyle<FTableRowStyle>("TableView.Row"))
+		, _CategoryRowStyle(&FDataHierarchyEditorStyle::Get().GetWidgetStyle<FTableRowStyle>("HierarchyEditor.Row.Category"))
+	{}
+		SLATE_ARGUMENT(bool, bReadOnly)
+		SLATE_EVENT(FOnGenerateRowContentWidget, OnGenerateRowContentWidget)
+		SLATE_EVENT(FOnGenerateCustomDetailsPanelNameWidget, OnGenerateCustomDetailsPanelNameWidget)
+		SLATE_STYLE_ARGUMENT(FTableRowStyle, ItemRowStyle)
+		SLATE_STYLE_ARGUMENT(FTableRowStyle, CategoryRowStyle)
+	SLATE_END_ARGS()
+
+public:
+	void Construct(const FArguments& InArgs, TObjectPtr<UDataHierarchyViewModelBase> InHierarchyViewModel);
+	virtual ~SDataHierarchyEditor() override;
+
+	void RefreshSourceItems();
+	void RefreshAllViews(bool bFullRefresh = false);
+	void RequestRefreshAllViewsNextFrame(bool bFullRefresh = false);
+	void RefreshSourceView(bool bFullRefresh = false) const;
+	void RequestRefreshSourceViewNextFrame(bool bFullRefresh = false);
+	void RefreshHierarchyView(bool bFullRefresh = false) const;
+	void RequestRefreshHierarchyViewNextFrame(bool bFullRefresh = false);
+	void RefreshSectionsView();
+	void RequestRefreshSectionsViewNextFrame();
+
+	void NavigateToHierarchyElement(FHierarchyElementIdentity Identity) const;
+	void NavigateToHierarchyElement(TSharedPtr<FHierarchyElementViewModel> Item) const;
+	bool IsItemSelected(TSharedPtr<FHierarchyElementViewModel> Item) const;
+	
+private:
+	// need to do this to enable focus so we can handle shortcuts
+	virtual bool SupportsKeyboardFocus() const override { return true; }
+	
+	TSharedRef<ITableRow> GenerateSourceItemRow(TSharedPtr<FHierarchyElementViewModel> HierarchyItem, const TSharedRef<STableViewBase>& TableViewBase);
+	TSharedRef<ITableRow> GenerateHierarchyItemRow(TSharedPtr<FHierarchyElementViewModel> HierarchyItem, const TSharedRef<STableViewBase>& TableViewBase);
+
+	bool FilterForSourceSection(TSharedPtr<const FHierarchyElementViewModel> ItemViewModel) const;
+private:
+	void Reinitialize();
+
+	void BindToHierarchyRootViewModel();
+	void UnbindFromHierarchyRootViewModel() const;
+	
+	/** Source items reflect the base, unedited status of items to edit into a hierarchy */
+	const TArray<TSharedPtr<FHierarchyElementViewModel>>& GetSourceItems() const;
+	
+	bool IsDetailsPanelEditingAllowed() const;
+	
+	void ClearSourceItems() const;
+	
+	void RequestRenameSelectedItem();
+	bool CanRequestRenameSelectedItem() const;
+
+	void DeleteItems(TArray<TSharedPtr<FHierarchyElementViewModel>> ItemsToDelete) const;
+	void DeleteSelectedHierarchyItems() const;
+	bool CanDeleteSelectedElements() const;
+
+	void NavigateToMatchingHierarchyElementFromSelectedSourceElement() const;
+	bool CanNavigateToMatchingHierarchyElementFromSelectedSourceElement() const;
+
+	void DeleteActiveSection() const;
+	bool CanDeleteActiveSection() const;
+
+	void OnElementAdded(TSharedPtr<FHierarchyElementViewModel> AddedItem);
+	void OnHierarchySectionActivated(TSharedPtr<FHierarchySectionViewModel> Section);
+	void OnSourceSectionActivated(TSharedPtr<FHierarchySectionViewModel> Section);
+	void OnHierarchySectionAdded(TSharedPtr<FHierarchySectionViewModel> AddedSection);
+	void OnHierarchySectionDeleted(TSharedPtr<FHierarchySectionViewModel> DeletedSection);
+
+	void SetActiveSourceSection(TSharedPtr<struct FHierarchySectionViewModel>);
+	TSharedPtr<FHierarchySectionViewModel> GetActiveSourceSection() const;
+	UHierarchySection* GetActiveSourceSectionData() const;
+	
+	void OnSelectionChanged(TSharedPtr<FHierarchyElementViewModel> HierarchyItem, ESelectInfo::Type Type, bool bFromHierarchy) const;
+
+	void RunSourceSearch();
+	void OnSourceSearchTextChanged(const FText& Text);
+	void OnSourceSearchTextCommitted(const FText& Text, ETextCommit::Type CommitType);
+	void OnSearchButtonClicked(SSearchBox::SearchDirection SearchDirection);
+
+	virtual FReply OnKeyDown(const FGeometry& MyGeometry, const FKeyEvent& InKeyEvent) override;
+	virtual FReply OnMouseButtonDown(const FGeometry& MyGeometry, const FPointerEvent& MouseEvent) override;
+	virtual FReply OnMouseButtonUp(const FGeometry& MyGeometry, const FPointerEvent& MouseEvent) override;
+
+	FReply OnAddCategoryClicked() const;
+	FReply OnAddSectionClicked() const;
+
+	TSharedPtr<SWidget> SummonContextMenuForSelectedRows(bool bFromHierarchy) const;
+
+	struct FSearchItem
+	{
+		TArray<TSharedPtr<FHierarchyElementViewModel>> Path;
+
+		TSharedPtr<FHierarchyElementViewModel> GetEntry() const
+		{
+			return Path.Num() > 0 ? 
+				Path[Path.Num() - 1] : 
+				nullptr;
+		}
+
+		bool operator==(const FSearchItem& Item) const
+		{
+			return Path == Item.Path;
+		}
+	};
+
+	/** This will recursively generated parent chain paths for all items within the given root. Used for expansion purposes. */
+	void GenerateSearchItems(TSharedRef<FHierarchyElementViewModel> Root, TArray<TSharedPtr<FHierarchyElementViewModel>> ParentChain, TArray<FSearchItem>& OutSearchItems);
+	void ExpandSourceSearchResults();
+	void SelectNextSourceSearchResult();
+	void SelectPreviousSourceSearchResult();
+	TOptional<SSearchBox::FSearchResultData> GetSearchResultData() const;
+	
+	FHierarchyElementViewModel::FCanPerformActionResults CanDropOnRoot(TSharedPtr<FHierarchyElementViewModel> DraggedItem) const;
+
+	/** Callback functions for the root widget */
+	FReply HandleHierarchyRootDrop(const FGeometry&, const FDragDropEvent& DragDropEvent) const;
+	bool OnCanDropOnRoot(TSharedPtr<FDragDropOperation> DragDropOperation) const;
+	void OnRootDragEnter(const FDragDropEvent& DragDropEvent) const;
+	void OnRootDragLeave(const FDragDropEvent& DragDropEvent) const;
+	FSlateColor GetRootIconColor() const;
+
+	virtual void NotifyPostChange(const FPropertyChangedEvent& PropertyChangedEvent, FProperty* PropertyThatChanged) override;
+private:
+	TWeakObjectPtr<UDataHierarchyViewModelBase> HierarchyViewModel;
+	
+	TArray<FSearchItem> SourceSearchResults;
+	TOptional<FSearchItem> FocusedSearchResult;
+
+	const FTableRowStyle* CategoryRowStyle = nullptr;
+	const FTableRowStyle* ItemRowStyle = nullptr;
+	
+	mutable TWeakPtr<FHierarchyElementViewModel> SelectedDetailsPanelItemViewModel;
+
+private:
+	TStrongObjectPtr<UHierarchyRoot> SourceRoot;
+	TSharedPtr<FHierarchyRootViewModel> SourceRootViewModel;
+	TSharedPtr<STreeView<TSharedPtr<FHierarchyElementViewModel>>> SourceTreeView;
+	TSharedPtr<STreeView<TSharedPtr<FHierarchyElementViewModel>>> HierarchyTreeView;
+	TWeakPtr<struct FHierarchySectionViewModel> ActiveSourceSection;
+	TSharedPtr<class SWrapBox> SourceSectionBox;
+	TSharedPtr<class SWrapBox> HierarchySectionBox;
+	TSharedPtr<SSearchBox> SourceSearchBox;
+	TSharedPtr<IDetailsView> DetailsPanel;
+private:
+	FOnGenerateRowContentWidget OnGenerateRowContentWidget;
+	FOnGenerateCustomDetailsPanelNameWidget OnGenerateCustomDetailsPanelNameWidget;
+	TSharedPtr<FActiveTimerHandle> RefreshHierarchyViewNextFrameHandle;
+	TSharedPtr<FActiveTimerHandle> RefreshSourceViewNextFrameHandle;
+	TSharedPtr<FActiveTimerHandle> RefreshSectionsViewNextFrameHandle;
+};
+
+class DATAHIERARCHYEDITOR_API SHierarchyCategory : public SCompoundWidget
+{
+	SLATE_BEGIN_ARGS(SHierarchyCategory)
+	{}
+		SLATE_EVENT(FIsSelected, IsSelected)
+	SLATE_END_ARGS()
+
+public:
+	void Construct(const FArguments& InArgs, TSharedPtr<struct FHierarchyCategoryViewModel> InCategoryViewModel);
+
+	void EnterEditingMode() const;
+	bool OnVerifyCategoryRename(const FText& NewName, FText& OutTooltip) const;
+	
+	FText GetCategoryText() const;
+	void OnRenameCategory(const FText& NewText, ETextCommit::Type) const;
+	
+private:
+	TWeakPtr<struct FHierarchyCategoryViewModel> CategoryViewModel;
+	TSharedPtr<SInlineEditableTextBlock> InlineEditableTextBlock;
+};
+
+class SSectionDragDropTarget : public SDropTarget
+{
+public:
+	SLATE_BEGIN_ARGS(SSectionDragDropTarget)
+		: _DropTargetArgs(SDropTarget::FArguments())
+	{}	
+		SLATE_ARGUMENT(SDropTarget::FArguments, DropTargetArgs)
+	SLATE_END_ARGS()
+
+	void Construct(const FArguments& InArgs, TSharedPtr<FHierarchySectionViewModel> OwningSection, EItemDropZone ItemDropZone);
+	
+	virtual FReply OnDragOver(const FGeometry& MyGeometry, const FDragDropEvent& DragDropEvent) override;
+	virtual void OnDragLeave(const FDragDropEvent& DragDropEvent) override;
+
+private:
+	TSharedPtr<FHierarchySectionViewModel> OwningSection;
+	/** The drop zone this drop target represents */
+	EItemDropZone DropZone = EItemDropZone::OntoItem;
+};
diff --git a/Engine/Source/Editor/DataHierarchyEditor/Private/DataHierarchyEditorCommands.cpp b/Engine/Source/Editor/DataHierarchyEditor/Private/DataHierarchyEditorCommands.cpp
new file mode 100644
--- /dev/null
+++ b/Engine/Source/Editor/DataHierarchyEditor/Private/DataHierarchyEditorCommands.cpp
@@ -0,0 +1,12 @@
+// Copyright Epic Games, Inc. All Rights Reserved.
+
+#include "DataHierarchyEditorCommands.h"
+
+#define LOCTEXT_NAMESPACE "DataHierarchyEditorCommands"
+
+void FDataHierarchyEditorCommands::RegisterCommands()
+{
+	UI_COMMAND(FindInHierarchy, "Find in Hierarchy", "Find the element in the hierarchy", EUserInterfaceActionType::Button, FInputChord(EKeys::F));
+}
+
+#undef LOCTEXT_NAMESPACE
diff --git a/Engine/Source/Editor/DataHierarchyEditor/Private/DataHierarchyEditorMisc.cpp b/Engine/Source/Editor/DataHierarchyEditor/Private/DataHierarchyEditorMisc.cpp
new file mode 100644
--- /dev/null
+++ b/Engine/Source/Editor/DataHierarchyEditor/Private/DataHierarchyEditorMisc.cpp
@@ -0,0 +1,37 @@
+﻿// Copyright Epic Games, Inc. All Rights Reserved.
+
+#include "DataHierarchyEditorMisc.h"
+
+FName UE::DataHierarchyEditor::GetUniqueName(FName CandidateName, const TSet<FName>& ExistingNames)
+{
+	// This utility function needs to generate a unique name while only considering the text portion of the name and
+	// not the index, so generate names with 0 indices before using them for comparison.
+	TSet<FName> ExistingNamesWithIndexZero;
+	for (FName ExistingName : ExistingNames)
+	{
+		ExistingNamesWithIndexZero.Add(FName(ExistingName, 0));
+	}
+	FName CandidateNameWithIndexZero = FName(CandidateName, 0);
+
+	if (ExistingNamesWithIndexZero.Contains(CandidateNameWithIndexZero) == false)
+	{
+		return CandidateName;
+	}
+
+	FString CandidateNameString = CandidateNameWithIndexZero.ToString();
+	FString BaseNameString = CandidateNameString;
+	if (CandidateNameString.Len() >= 3 && CandidateNameString.Right(3).IsNumeric())
+	{
+		BaseNameString = CandidateNameString.Left(CandidateNameString.Len() - 3);
+	}
+
+	FName UniqueName = FName(*BaseNameString);
+	int32 NameIndex = 1;
+	while (ExistingNamesWithIndexZero.Contains(UniqueName))
+	{
+		UniqueName = FName(*FString::Printf(TEXT("%s%03i"), *BaseNameString, NameIndex));
+		NameIndex++;
+	}
+
+	return UniqueName;
+}
diff --git a/Engine/Source/Editor/DataHierarchyEditor/Private/DataHierarchyEditorModule.cpp b/Engine/Source/Editor/DataHierarchyEditor/Private/DataHierarchyEditorModule.cpp
new file mode 100644
--- /dev/null
+++ b/Engine/Source/Editor/DataHierarchyEditor/Private/DataHierarchyEditorModule.cpp
@@ -0,0 +1,26 @@
+﻿// Copyright Epic Games, Inc. All Rights Reserved.
+
+#include "DataHierarchyEditorModule.h"
+
+#include "DataHierarchyEditorCommands.h"
+#include "DataHierarchyEditorStyle.h"
+
+#define LOCTEXT_NAMESPACE "DataHierarchyEditor"
+
+DEFINE_LOG_CATEGORY(LogDataHierarchyEditor)
+
+void FDataHierarchyEditorModule::StartupModule()
+{
+    FDataHierarchyEditorStyle::Register();
+	FDataHierarchyEditorCommands::Register();
+}
+
+void FDataHierarchyEditorModule::ShutdownModule()
+{
+	FDataHierarchyEditorCommands::Unregister();
+    FDataHierarchyEditorStyle::Shutdown();
+}
+
+#undef LOCTEXT_NAMESPACE
+    
+IMPLEMENT_MODULE(FDataHierarchyEditorModule, DataHierarchyEditor)
diff --git a/Engine/Source/Editor/DataHierarchyEditor/Private/DataHierarchyEditorStyle.cpp b/Engine/Source/Editor/DataHierarchyEditor/Private/DataHierarchyEditorStyle.cpp
new file mode 100644
--- /dev/null
+++ b/Engine/Source/Editor/DataHierarchyEditor/Private/DataHierarchyEditorStyle.cpp
@@ -0,0 +1,126 @@
+// Copyright Epic Games, Inc. All Rights Reserved.
+
+#include "DataHierarchyEditorStyle.h"
+
+#include "Framework/Application/SlateApplication.h"
+#include "Styling/SlateStyleMacros.h"
+#include "Styling/AppStyle.h"
+#include "Styling/SlateStyleRegistry.h"
+#include "Styling/SlateTypes.h"
+#include "Styling/CoreStyle.h"
+#include "Styling/StarshipCoreStyle.h"
+#include "Styling/StyleColors.h"
+
+TSharedPtr<FDataHierarchyEditorStyle> FDataHierarchyEditorStyle::DataHierarchyEditorStyle = nullptr;
+
+void FDataHierarchyEditorStyle::Register()
+{
+	FSlateStyleRegistry::RegisterSlateStyle(Get());
+}
+
+void FDataHierarchyEditorStyle::Unregister()
+{
+	FSlateStyleRegistry::UnRegisterSlateStyle(Get());
+}
+
+void FDataHierarchyEditorStyle::Shutdown()
+{
+	Unregister();
+	DataHierarchyEditorStyle.Reset();
+}
+
+void FDataHierarchyEditorStyle::InitDataHierarchyEditor()
+{
+	const FTextBlockStyle NormalText = FAppStyle::GetWidgetStyle<FTextBlockStyle>("NormalText");
+	FSlateFontInfo CategoryFont = FAppStyle::Get().GetFontStyle(TEXT("DetailsView.CategoryFontStyle"));
+	CategoryFont.Size = 11;
+	
+	const FEditableTextBoxStyle NormalEditableTextBox = FAppStyle::Get().GetWidgetStyle<FEditableTextBoxStyle>("NormalEditableTextBox");
+	
+	FEditableTextBoxStyle CategoryEditableText = FEditableTextBoxStyle(NormalEditableTextBox)
+		.SetFont(CategoryFont)
+		.SetForegroundColor(FStyleColors::AccentWhite);
+	
+	FTextBlockStyle CategoryText = FTextBlockStyle(NormalText)
+		.SetFont(CategoryFont);
+	FInlineEditableTextBlockStyle HierarchyCategoryTextStyle = FInlineEditableTextBlockStyle()
+		.SetTextStyle(CategoryText)
+		.SetEditableTextBoxStyle(CategoryEditableText);
+
+	Set("HierarchyEditor.CategoryTextBlock", CategoryText);
+	Set("HierarchyEditor.Category", HierarchyCategoryTextStyle);
+
+	Set("HierarchyEditor.Stack.DropTarget.BorderVertical", new IMAGE_BRUSH("Icons/StackDropTargetBorder_Vertical", FVector2D(2, 8), FLinearColor::White, ESlateBrushTileType::Vertical));
+	Set("HierarchyEditor.Stack.DropTarget.BorderHorizontal", new IMAGE_BRUSH("Icons/StackDropTargetBorder_Horizontal", FVector2D(8, 2), FLinearColor::White, ESlateBrushTileType::Horizontal));
+
+	FButtonStyle SimpleButtonStyle = FAppStyle::Get().GetWidgetStyle<FButtonStyle>("SimpleButton");
+	FButtonStyle ButtonStyle = FButtonStyle(SimpleButtonStyle)
+		.SetNormalForeground(FStyleColors::Foreground)
+		.SetPressedForeground(FStyleColors::ForegroundHover)
+		.SetHoveredForeground(FStyleColors::ForegroundHover)
+		.SetDisabledForeground(FStyleColors::Foreground)
+		.SetNormalPadding(FMargin(8.f, 2.f, 8.f, 2.f))
+		.SetPressedPadding(FMargin(8.f, 3.f, 8.f, 1.f));
+
+	Set("HierarchyEditor.ButtonStyle", ButtonStyle);
+
+	const FSlateColor SelectorColor = FAppStyle::GetSlateColor("SelectorColor");
+	const FSlateColor SelectionColor = FAppStyle::GetSlateColor("SelectionColor");
+	const FSlateColor SelectionColor_Inactive = FAppStyle::GetSlateColor("SelectionColor_Inactive");
+	FLinearColor HighlightColor = FLinearColor(0.068f, 0.068f, 0.068f);
+	
+	Set("HierarchyEditor.Color.Highlight", HighlightColor);
+	
+	FTableRowStyle HierarchyEditorCategoryRowStyle = FTableRowStyle()
+		.SetEvenRowBackgroundBrush(*FAppStyle::GetBrush("DetailsView.CategoryTop"))
+		.SetOddRowBackgroundBrush(*FAppStyle::GetBrush("DetailsView.CategoryTop"))
+		.SetEvenRowBackgroundHoveredBrush(*FAppStyle::GetBrush("DetailsView.CategoryTop_Hovered"))
+		.SetOddRowBackgroundHoveredBrush(*FAppStyle::GetBrush("DetailsView.CategoryTop_Hovered"))
+		.SetSelectorFocusedBrush( CORE_BORDER_BRUSH( "Common/Selector", FMargin(4.f/16.f), SelectorColor))
+		.SetActiveBrush( CORE_IMAGE_BRUSH( "Common/Selection", CoreStyleConstants::Icon8x8, SelectionColor))
+		.SetActiveHoveredBrush( CORE_IMAGE_BRUSH( "Common/Selection", CoreStyleConstants::Icon8x8, SelectionColor))
+		.SetInactiveBrush( CORE_IMAGE_BRUSH( "Common/Selection", CoreStyleConstants::Icon8x8, SelectionColor_Inactive))
+		.SetInactiveHoveredBrush( CORE_IMAGE_BRUSH( "Common/Selection", CoreStyleConstants::Icon8x8, SelectionColor_Inactive))
+		.SetActiveHighlightedBrush( CORE_IMAGE_BRUSH("Common/Selection", CoreStyleConstants::Icon8x8, HighlightColor))
+		.SetInactiveHighlightedBrush( CORE_IMAGE_BRUSH("Common/Selection", CoreStyleConstants::Icon8x8, HighlightColor))
+		.SetDropIndicator_Above(CORE_BOX_BRUSH("Common/DropZoneIndicator_Above", FMargin(10.0f / 16.0f, 10.0f / 16.0f, 0, 0), FStarshipCoreStyle::GetCoreStyle().GetSlateColor("SelectorColor")))
+		.SetDropIndicator_Onto(CORE_BOX_BRUSH("Common/DropZoneIndicator_Onto", FMargin(4.0f / 16.0f), FStarshipCoreStyle::GetCoreStyle().GetSlateColor("SelectorColor")))
+		.SetDropIndicator_Below(CORE_BOX_BRUSH("Common/DropZoneIndicator_Below", FMargin(10.0f / 16.0f, 0, 0, 10.0f / 16.0f), FStarshipCoreStyle::GetCoreStyle().GetSlateColor("SelectorColor")));
+	
+	Set("HierarchyEditor.Row.Category", HierarchyEditorCategoryRowStyle);
+	
+	FInlineEditableTextBlockStyle HierarchyEditorCategoryStyle = FAppStyle::GetWidgetStyle<FInlineEditableTextBlockStyle>("InlineEditableTextBlockStyle");
+	HierarchyEditorCategoryStyle.TextStyle = FAppStyle::Get().GetWidgetStyle<FTextBlockStyle>("DetailsView.CategoryTextStyle");
+	HierarchyEditorCategoryStyle.TextStyle.SetFontSize(10.f);
+	Set("HierarchyEditor.CategoryTextStyle", HierarchyEditorCategoryStyle);
+	
+	Set("HierarchyEditor.RootDropIcon", new IMAGE_BRUSH_SVG("Starship/Common/caret-down", CoreStyleConstants::Icon20x20));
+}
+
+FDataHierarchyEditorStyle::FDataHierarchyEditorStyle() : FSlateStyleSet("DataHierarchyEditorStyle")
+{
+	SetContentRoot(FPaths::EngineContentDir() / TEXT("Slate"));
+	InitDataHierarchyEditor();
+}
+
+void FDataHierarchyEditorStyle::ReloadTextures()
+{
+	FSlateApplication::Get().GetRenderer()->ReloadTextureResources();
+}
+
+const FDataHierarchyEditorStyle& FDataHierarchyEditorStyle::Get()
+{
+	if(!DataHierarchyEditorStyle.IsValid())
+	{
+		DataHierarchyEditorStyle = MakeShareable(new FDataHierarchyEditorStyle());
+	}
+	
+	return *DataHierarchyEditorStyle;
+}
+
+void FDataHierarchyEditorStyle::ReinitializeStyle()
+{
+	Unregister();
+	DataHierarchyEditorStyle.Reset();
+	Register();	
+}
diff --git a/Engine/Source/Editor/DataHierarchyEditor/Private/DataHierarchyViewModelBase.cpp b/Engine/Source/Editor/DataHierarchyEditor/Private/DataHierarchyViewModelBase.cpp
new file mode 100644
--- /dev/null
+++ b/Engine/Source/Editor/DataHierarchyEditor/Private/DataHierarchyViewModelBase.cpp
@@ -0,0 +1,1880 @@
+// Copyright Epic Games, Inc. All Rights Reserved.
+
+#include "DataHierarchyViewModelBase.h"
+#include "SDropTarget.h"
+#include "Widgets/Views/STableRow.h"
+#include "Widgets/SDataHierarchyEditor.h"
+#include "Widgets/Text/SInlineEditableTextBlock.h"
+#include "Widgets/SToolTip.h"
+#include "Widgets/Input/SCheckBox.h"
+#include "Editor.h"
+#include "DataHierarchyEditorCommands.h"
+#include "DataHierarchyEditorMisc.h"
+#include "DataHierarchyEditorModule.h"
+#include "IPropertyRowGenerator.h"
+#include "ScopedTransaction.h"
+#include "Framework/Commands/GenericCommands.h"
+#include "ScopedTransaction.h"
+#include "Framework/Notifications/NotificationManager.h"
+#include "ToolMenus.h"
+#include "Logging/StructuredLog.h"
+#include "Widgets/Notifications/SNotificationList.h"
+
+#define LOCTEXT_NAMESPACE "DataHierarchyEditor"
+
+const TArray<UHierarchyElement*>& UHierarchyElement::GetChildren() const
+{
+	return Children;
+}
+
+UHierarchyElement* UHierarchyElement::FindChildWithIdentity(FHierarchyElementIdentity ChildIdentity, bool bSearchRecursively)
+{
+	TObjectPtr<UHierarchyElement>* FoundItem = Children.FindByPredicate([ChildIdentity](UHierarchyElement* Child)
+	{
+		return Child->GetPersistentIdentity() == ChildIdentity;
+	});
+
+	if(FoundItem)
+	{
+		return *FoundItem;
+	}
+	
+	if(bSearchRecursively)
+	{
+		for(UHierarchyElement* Child : Children)
+		{
+			UHierarchyElement* FoundChild = Child->FindChildWithIdentity(ChildIdentity, bSearchRecursively);
+
+			if(FoundChild)
+			{
+				return FoundChild;
+			}
+		}
+	}	
+
+	return nullptr;
+}
+
+UHierarchyElement* UHierarchyElement::CopyAndAddItemAsChild(const UHierarchyElement& ItemToCopy)
+{
+	UHierarchyElement* NewChild = Cast<UHierarchyElement>(StaticDuplicateObject(&ItemToCopy, this));
+	if(NewChild->GetPersistentIdentity() != ItemToCopy.GetPersistentIdentity())
+	{
+		check(false);
+	}
+	GetChildrenMutable().Add(NewChild);
+
+	return NewChild;
+}
+
+UHierarchyElement* UHierarchyElement::CopyAndAddItemUnderParentIdentity(const UHierarchyElement& ItemToCopy, FHierarchyElementIdentity ParentIdentity)
+{
+	UHierarchyElement* ParentItem = FindChildWithIdentity(ParentIdentity, true);
+
+	if(ParentItem)
+	{
+		UHierarchyElement* NewChild = Cast<UHierarchyElement>(StaticDuplicateObject(&ItemToCopy, ParentItem));
+		if(NewChild->GetPersistentIdentity() != ItemToCopy.GetPersistentIdentity())
+		{
+			check(false);
+		}
+		ParentItem->GetChildrenMutable().Add(NewChild);
+		return NewChild;
+	}
+
+	return nullptr;
+}
+
+bool UHierarchyElement::RemoveChildWithIdentity(FHierarchyElementIdentity ChildIdentity, bool bSearchRecursively)
+{	
+	int32 RemovedChildrenCount = Children.RemoveAll([ChildIdentity](UHierarchyElement* Child)
+	{
+		return Child->GetPersistentIdentity() == ChildIdentity;
+	});
+
+	if(RemovedChildrenCount > 1)
+	{
+		UE_LOG(LogDataHierarchyEditor, Warning, TEXT("More than one child with the same identity has been found in parent %s"), *ToString());
+	}
+	
+	bool bChildrenRemoved = RemovedChildrenCount > 0;
+	
+	if(bSearchRecursively && bChildrenRemoved == false)
+	{
+		for(UHierarchyElement* Child : Children)
+		{
+			bChildrenRemoved |= Child->RemoveChildWithIdentity(ChildIdentity, bSearchRecursively);
+		}
+	}
+
+	return bChildrenRemoved;
+}
+
+TArray<FHierarchyElementIdentity> UHierarchyElement::GetParentIdentities() const
+{
+	TArray<FHierarchyElementIdentity> ParentIdentities;
+
+	for(UHierarchyElement* Parent = Cast<UHierarchyElement>(GetOuter()); Parent != nullptr; Parent = Cast<UHierarchyElement>(Parent->GetOuter()))
+	{
+		ParentIdentities.Add(Parent->GetPersistentIdentity());
+	}
+
+	return ParentIdentities;
+}
+
+bool UHierarchyElement::Modify(bool bAlwaysMarkDirty)
+{
+	bool bSavedToTransactionBuffer = true;
+
+	for(UHierarchyElement* Child : Children)
+	{
+		bSavedToTransactionBuffer &= Child->Modify(bAlwaysMarkDirty);
+	}
+	
+	bSavedToTransactionBuffer &= UObject::Modify(bAlwaysMarkDirty);
+
+	return bSavedToTransactionBuffer;
+}
+
+void UHierarchyElement::PostLoad()
+{
+	if(Guid_DEPRECATED.IsValid())
+	{
+		SetIdentity(FHierarchyElementIdentity({Guid_DEPRECATED}, {}));
+	}
+
+	bool bAnyChildNullptr = false;		
+	for(auto It = Children.CreateIterator(); It; ++It)
+	{
+		if(*It == nullptr)
+		{
+			bAnyChildNullptr = true;
+			It.RemoveCurrent();
+		}
+	}
+
+	if(bAnyChildNullptr)
+	{
+		UPackage* Package = GetPackage();
+		UE_LOG(LogDataHierarchyEditor, Warning, TEXT("HierarchyElement %s found nullptr child in asset %s. Removed all nullptr children. This is indicative of something wrong. Check if the hierarchy is still correct and fix it, if necessary."), *ToString(), *GetNameSafe(Package))	
+	}
+	
+	Super::PostLoad();
+}
+
+UHierarchySection* UHierarchyRoot::AddSection(FText InNewSectionName, int32 InsertIndex, TSubclassOf<UHierarchySection> SectionClass)
+{
+	TSet<FName> ExistingSectionNames;
+	
+	for(FName& SectionName : GetSections())
+	{
+		ExistingSectionNames.Add(SectionName);
+	}
+	
+	FName NewName = UE::DataHierarchyEditor::GetUniqueName(FName(InNewSectionName.ToString()), ExistingSectionNames);
+	UHierarchySection* NewSectionItem = NewObject<UHierarchySection>(this, SectionClass);
+	NewSectionItem->SetSectionName(NewName);
+	NewSectionItem->SetFlags(RF_Transactional);
+	
+	if(InsertIndex == INDEX_NONE)
+	{
+		Sections.Add(NewSectionItem);
+	}
+	else
+	{
+		Sections.Insert(NewSectionItem, InsertIndex);
+	}
+	
+	return NewSectionItem;
+}
+
+UHierarchySection* UHierarchyRoot::FindSectionByIdentity(FHierarchyElementIdentity SectionIdentity)
+{
+	for(UHierarchySection* Section : Sections)
+	{
+		if(Section->GetPersistentIdentity() == SectionIdentity)
+		{
+			return Section;
+		}
+	}
+
+	return nullptr;
+}
+
+void UHierarchyRoot::DuplicateSectionFromOtherRoot(const UHierarchySection& SectionToCopy)
+{
+	if(FindSectionByIdentity(SectionToCopy.GetPersistentIdentity()) != nullptr || SectionToCopy.GetOuter() == this)
+	{
+		return;
+	}
+	
+	Sections.Add(Cast<UHierarchySection>(StaticDuplicateObject(&SectionToCopy, this)));
+}
+
+void UHierarchyRoot::RemoveSection(FText SectionName)
+{
+	if(Sections.ContainsByPredicate([SectionName](UHierarchySection* Section)
+	{
+		return Section->GetSectionNameAsText().EqualTo(SectionName);
+	}))
+	{
+		Sections.RemoveAll([SectionName](UHierarchySection* Section)
+		{
+			return Section->GetSectionNameAsText().EqualTo(SectionName);
+		});
+	}
+}
+
+void UHierarchyRoot::RemoveSectionByIdentity(FHierarchyElementIdentity SectionIdentity)
+{
+	Sections.RemoveAll([SectionIdentity](UHierarchySection* Section)
+	{
+		return Section->GetPersistentIdentity() == SectionIdentity;
+	});
+}
+
+TSet<FName> UHierarchyRoot::GetSections() const
+{
+	TSet<FName> OutSections;
+	for(UHierarchySection* Section : Sections)
+	{
+		OutSections.Add(Section->GetSectionName());
+	}
+
+	return OutSections;
+}
+
+int32 UHierarchyRoot::GetSectionIndex(UHierarchySection* Section) const
+{
+	return Sections.Find(Section);
+}
+
+bool UHierarchyRoot::Modify(bool bAlwaysMarkDirty)
+{
+	bool bSavedToTransactionBuffer = true;
+	
+	for(UHierarchySection* Section : Sections)
+	{
+		bSavedToTransactionBuffer &= Section->Modify();
+	}
+	
+	bSavedToTransactionBuffer &= Super::Modify(bAlwaysMarkDirty);	
+
+	return bSavedToTransactionBuffer;
+}
+
+void UHierarchyRoot::EmptyAllData()
+{
+	Children.Empty();
+	Sections.Empty();
+}
+
+void UHierarchyRoot::Serialize(FStructuredArchive::FRecord Record)
+{
+	// If the root isn't transient, neither should any of its hierarchy elements be.
+	// This is expected to happen as the source elements are transient by default.
+	// When source hierarchy elements are put into the hierarchy we have to make sure to remove the flag after
+	if(Record.GetArchiveState().IsSaving() && this->HasAnyFlags(RF_Transient) == false)
+	{
+		TArray<UHierarchyElement*> AllElements;
+		GetChildrenOfType(AllElements, true);
+
+		for(UHierarchyElement* Element : AllElements)
+		{
+			Element->ClearFlags(RF_Transient);
+		}
+	}
+	
+	Super::Serialize(Record);
+}
+
+bool FHierarchyCategoryViewModel::IsTopCategoryActive() const
+{
+	if(UHierarchyCategory* Category = GetDataMutable<UHierarchyCategory>())
+	{
+		const UHierarchyCategory* Result = Category;
+		const UHierarchyCategory* TopLevelCategory = Result;
+		
+		for (; TopLevelCategory != nullptr; TopLevelCategory = TopLevelCategory->GetTypedOuter<UHierarchyCategory>() )
+		{
+			if(TopLevelCategory != nullptr)
+			{
+				Result = TopLevelCategory;
+			}
+		}
+		
+		return HierarchyViewModel->IsHierarchySectionActive(Result->GetSection());
+	}
+
+	return false;	
+}
+
+FHierarchyElementViewModel::FCanPerformActionResults FHierarchyCategoryViewModel::CanDropOnInternal(TSharedPtr<FHierarchyElementViewModel> DraggedElement, EItemDropZone ItemDropZone)
+{
+	FCanPerformActionResults Results(false);
+	
+	TArray<TSharedPtr<FHierarchyCategoryViewModel>> TargetChildrenCategories;
+	GetChildrenViewModelsForType<UHierarchyCategory, FHierarchyCategoryViewModel>(TargetChildrenCategories);
+	
+	TArray<TSharedPtr<FHierarchyCategoryViewModel>> SiblingCategories;
+	Parent.Pin()->GetChildrenViewModelsForType<UHierarchyCategory, FHierarchyCategoryViewModel>(SiblingCategories);
+	
+	// we only allow drops if some general conditions are fulfilled
+	if(DraggedElement->GetData() != GetData() &&
+		(!DraggedElement->HasParent(AsShared(), false) || ItemDropZone != EItemDropZone::OntoItem)  &&
+		!HasParent(DraggedElement, true))
+	{
+		// categories can be dropped on categories, but only if the resulting sibling categories or children categories have different names
+		if(DraggedElement->GetData()->IsA<UHierarchyCategory>())
+		{
+			if(ItemDropZone != EItemDropZone::OntoItem)
+			{
+				bool bContainsSiblingWithSameName = SiblingCategories.ContainsByPredicate([DraggedElement](TSharedPtr<FHierarchyCategoryViewModel> HierarchyCategoryViewModel)
+					{
+						return DraggedElement->ToString() == HierarchyCategoryViewModel->ToString() && DraggedElement != HierarchyCategoryViewModel;
+					});
+
+				if(bContainsSiblingWithSameName)
+				{
+					Results.bCanPerform = false;
+					Results.CanPerformMessage = LOCTEXT("CantDropCategorNextToCategorySameSiblingNames", "A category of the same name already exists here, potentially in a different section. Please rename your category first.");
+					return Results;
+				}
+
+				Results.CanPerformMessage = LOCTEXT("MoveCategoryText", "Move category here");
+
+				// if we are making a category a sibling of another at the root level, the section will be set to the currently active section. Let that be known.
+				if(Parent.Pin()->GetData()->IsA<UHierarchyRoot>())
+				{
+					UHierarchyCategory* DraggedCategory = Cast<UHierarchyCategory>(DraggedElement->GetDataMutable());
+					if(DraggedCategory->GetSection() != HierarchyViewModel->GetActiveHierarchySectionData())
+					{
+						FText SectionChangeBaseText = LOCTEXT("CategorySectionWillUpdateDueToDrop", "The section of the category will change to {0} after the drop");
+						FText ActualSectionChangeText = FText::FormatOrdered(SectionChangeBaseText, HierarchyViewModel->GetActiveHierarchySectionData() == nullptr ? FText::FromString("All") : HierarchyViewModel->GetActiveHierarchySectionData()->GetSectionNameAsText());
+						Results.CanPerformMessage = FText::FormatOrdered(FText::AsCultureInvariant("{0}\n{1}"), Results.CanPerformMessage, ActualSectionChangeText);
+					}
+				}
+			}
+			else
+			{
+				bool bContainsChildrenCategoriesWithSameName = TargetChildrenCategories.ContainsByPredicate([DraggedElement](TSharedPtr<FHierarchyCategoryViewModel> HierarchyCategoryViewModel)
+					{
+						return DraggedElement->ToString() == HierarchyCategoryViewModel->ToString();
+					});
+
+				if(bContainsChildrenCategoriesWithSameName)
+				{
+					Results.bCanPerform = false;
+					Results.CanPerformMessage = LOCTEXT("CantDropCategoryOnCategorySameChildCategoryName", "A sub-category of the same name already exists! Please rename your category first.");
+					return Results;
+				}
+
+				Results.CanPerformMessage = LOCTEXT("CreateSubcategory", "Drop category here to create a sub-category");
+			}
+
+			Results.bCanPerform = true;
+			return Results;
+		}
+		else if(DraggedElement->GetData()->IsA<UHierarchyItem>())
+		{
+			// items can generally be dropped onto categories
+			Results.bCanPerform = EItemDropZone::OntoItem == ItemDropZone;
+
+			if(Results.bCanPerform)
+			{
+				if(DraggedElement->IsForHierarchy() == false)
+				{
+					FText Message = LOCTEXT("AddItemToCategoryDragMessage", "Add {0} to {1}");
+					Results.CanPerformMessage = FText::FormatOrdered(Message, FText::FromString(DraggedElement->ToString()), FText::FromString(ToString()));
+				}
+				else
+				{
+					FText Message = LOCTEXT("MoveItemToCategoryDragMessage", "Move {0} to {1}");
+					Results.CanPerformMessage = FText::FormatOrdered(Message, FText::FromString(DraggedElement->ToString()), FText::FromString(ToString()));
+				}
+			}
+		}
+	}
+
+	return Results;
+}
+
+void UHierarchyCategory::FixupSectionLinkage()
+{
+	UHierarchyRoot* OwningRoot = GetTypedOuter<UHierarchyRoot>();
+
+	if(Section != nullptr && Section->GetTypedOuter<UHierarchyRoot>() != OwningRoot)
+	{
+		UHierarchySection* CorrectSection = OwningRoot->FindSectionByIdentity(Section->GetPersistentIdentity());
+		ensure(CorrectSection != nullptr);
+		Section = CorrectSection;
+	}
+}
+
+void UHierarchySection::SetSectionNameAsText(const FText& Text)
+{
+	Section = FName(Text.ToString());
+}
+
+UDataHierarchyViewModelBase::UDataHierarchyViewModelBase()
+{
+	Commands = MakeShared<FUICommandList>();
+}
+
+UDataHierarchyViewModelBase::~UDataHierarchyViewModelBase()
+{
+	RefreshSourceViewDelegate.Unbind();
+	RefreshHierarchyWidgetDelegate.Unbind();
+	RefreshSectionsViewDelegate.Unbind();
+}
+
+void UDataHierarchyViewModelBase::Initialize()
+{		
+	HierarchyRoot = GetHierarchyRoot();
+	HierarchyRoot->SetFlags(RF_Transactional);
+
+	TArray<UHierarchyElement*> AllItems;
+	HierarchyRoot->GetChildrenOfType<UHierarchyElement>(AllItems, true);
+	for(UHierarchyElement* Item : AllItems)
+	{
+		Item->SetFlags(RF_Transactional);
+	}
+
+	for(UHierarchySection* Section : HierarchyRoot->GetSectionDataMutable())
+	{
+		Section->SetFlags(RF_Transactional);
+	}
+
+	UToolMenus* ToolMenus = UToolMenus::Get();
+
+	FName MenuName = GetContextMenuName();
+	if(!ToolMenus->IsMenuRegistered(MenuName))
+	{
+		UToolMenu* HierarchyMenu = ToolMenus->RegisterMenu(MenuName, NAME_None, EMultiBoxType::Menu);
+		HierarchyMenu->AddDynamicSection(NAME_None, FNewToolMenuDelegate::CreateStatic(&UDataHierarchyViewModelBase::GenerateDynamicContextMenu));
+	}
+	
+	SetupCommands();
+
+	TSharedPtr<FHierarchyElementViewModel> ViewModel = CreateViewModelForElement(HierarchyRoot, nullptr);
+	HierarchyRootViewModel = StaticCastSharedPtr<FHierarchyRootViewModel>(ViewModel); 
+	if(!ensureMsgf(HierarchyRootViewModel.IsValid(), TEXT("Make sure that CreateViewModelForData creates a FHierarchyRootViewModel (or derived) for UHierarchyRoot elements")))
+	{
+		return;
+	}
+	
+	HierarchyRootViewModel->Initialize();
+	HierarchyRootViewModel->AddChildFilter(FHierarchyElementViewModel::FOnFilterChild::CreateUObject(this, &UDataHierarchyViewModelBase::FilterForHierarchySection));
+	HierarchyRootViewModel->AddChildFilter(FHierarchyElementViewModel::FOnFilterChild::CreateUObject(this, &UDataHierarchyViewModelBase::FilterForUncategorizedRootItemsInAllSection));
+	HierarchyRootViewModel->SyncViewModelsToData();
+	
+	SetActiveHierarchySection(nullptr);
+	
+	InitializeInternal();
+
+	bIsInitialized = true;
+	
+	OnInitializedDelegate.ExecuteIfBound();
+}
+
+void UDataHierarchyViewModelBase::Finalize()
+{	
+	HierarchyRootViewModel.Reset();
+	HierarchyRoot = nullptr;
+	
+	FinalizeInternal();
+	
+	bIsFinalized = true;
+}
+
+TSharedPtr<FHierarchyElementViewModel> UDataHierarchyViewModelBase::CreateViewModelForElement(UHierarchyElement* Element, TSharedPtr<FHierarchyElementViewModel> Parent)
+{
+	// We first give the internal implementation a chance to create view models
+	if(TSharedPtr<FHierarchyElementViewModel> CustomViewModel = CreateCustomViewModelForElement(Element, Parent))
+	{
+		return CustomViewModel;
+	}
+
+	// If it wasn't implemented or wasn't covered, we make sure to have default view models
+	if(UHierarchyItem* Item = Cast<UHierarchyItem>(Element))
+	{
+		return MakeShared<FHierarchyItemViewModel>(Item, Parent.ToSharedRef(), this);
+	}
+	else if(UHierarchyCategory* Category = Cast<UHierarchyCategory>(Element))
+	{
+		return MakeShared<FHierarchyCategoryViewModel>(Category, Parent.ToSharedRef(), this);
+	}
+	else if(UHierarchySection* Section = Cast<UHierarchySection>(Element))
+	{
+		// For sections, we require the parent to be a root view model
+		TSharedPtr<FHierarchyRootViewModel> RootViewModel = StaticCastSharedPtr<FHierarchyRootViewModel>(Parent);
+		ensure(RootViewModel.IsValid());
+		return MakeShared<FHierarchySectionViewModel>(Section, RootViewModel.ToSharedRef(), this);
+	}
+	else if(UHierarchyRoot* Root = Cast<UHierarchyRoot>(Element))
+	{
+		// If the root is the hierarchy root, we know it's for the hierarchy. If not, it's the transient source root
+		bool bIsForHierarchy = GetHierarchyRoot() == Element;
+		return MakeShared<FHierarchyRootViewModel>(Root, this, bIsForHierarchy);
+	}
+
+	ensureMsgf(false, TEXT("This should never be reached. Either a custom or a default view model must exist for each Hierarchy Element"));
+	return nullptr;
+}
+
+TSubclassOf<UHierarchyCategory> UDataHierarchyViewModelBase::GetCategoryDataClass() const
+{
+	return UHierarchyCategory::StaticClass();
+}
+
+TSubclassOf<UHierarchySection> UDataHierarchyViewModelBase::GetSectionDataClass() const
+{
+	return UHierarchySection::StaticClass();
+}
+
+
+void UDataHierarchyViewModelBase::ForceFullRefresh()
+{
+	RefreshSourceItemsRequestedDelegate.ExecuteIfBound();
+	// todo (me) during merge at startup this can be nullptr for some reason
+	if(HierarchyRootViewModel.IsValid())
+	{
+		HierarchyRootViewModel->SyncViewModelsToData();
+	}
+	RefreshAllViewsRequestedDelegate.ExecuteIfBound(true);
+}
+
+void UDataHierarchyViewModelBase::ForceFullRefreshOnTimer()
+{
+	ensure(FullRefreshNextFrameHandle.IsValid());
+	ForceFullRefresh();
+	FullRefreshNextFrameHandle.Invalidate();
+}
+
+void UDataHierarchyViewModelBase::RequestFullRefreshNextFrame()
+{
+	if(!FullRefreshNextFrameHandle.IsValid() && GEditor != nullptr)
+	{
+		FullRefreshNextFrameHandle = GEditor->GetTimerManager()->SetTimerForNextTick(FTimerDelegate::CreateUObject(this, &UDataHierarchyViewModelBase::ForceFullRefreshOnTimer));
+	}
+}
+
+void UDataHierarchyViewModelBase::RefreshAllViews(bool bFullRefresh) const
+{
+	RefreshAllViewsRequestedDelegate.ExecuteIfBound(bFullRefresh);
+}
+
+void UDataHierarchyViewModelBase::RefreshSourceView(bool bFullRefresh) const
+{
+	RefreshSourceViewDelegate.ExecuteIfBound(bFullRefresh);
+}
+
+void UDataHierarchyViewModelBase::RefreshHierarchyView(bool bFullRefresh) const
+{
+	RefreshHierarchyWidgetDelegate.ExecuteIfBound(bFullRefresh);
+}
+
+void UDataHierarchyViewModelBase::RefreshSectionsView() const
+{
+	RefreshSectionsViewDelegate.ExecuteIfBound();
+}
+
+void UDataHierarchyViewModelBase::PostUndo(bool bSuccess)
+{
+	ForceFullRefresh();
+}
+
+void UDataHierarchyViewModelBase::PostRedo(bool bSuccess)
+{
+	PostUndo(bSuccess);
+}
+
+bool UDataHierarchyViewModelBase::MatchesContext(const FTransactionContext& InContext, const TArray<TPair<UObject*, FTransactionObjectEvent>>& TransactionObjectContexts) const
+{
+	for(const TPair<UObject*, FTransactionObjectEvent>& TransactionObjectContext : TransactionObjectContexts)
+	{
+		if(TransactionObjectContext.Key->IsA<UHierarchyElement>())
+		{
+			return true;
+		}
+	}
+	
+	return false;
+}
+
+bool UDataHierarchyViewModelBase::FilterForHierarchySection(TSharedPtr<const FHierarchyElementViewModel> ViewModel) const
+{
+	if(ActiveHierarchySection.IsValid())
+	{
+		return GetActiveHierarchySectionData() == ViewModel->GetSection();
+	}
+
+	return true;
+}
+
+bool UDataHierarchyViewModelBase::FilterForUncategorizedRootItemsInAllSection(TSharedPtr<const FHierarchyElementViewModel> ViewModel) const
+{
+	// we want to filter out all items that are directly added to the root if we aren't in the 'All' section
+	if(ActiveHierarchySection.IsValid())
+	{
+		return ViewModel->GetData<UHierarchyCategory>() != nullptr;
+	}
+
+	return true;
+}
+
+void UDataHierarchyViewModelBase::ToolMenuRequestRename(const FToolMenuContext& Context) const
+{
+	UHierarchyMenuContext* HierarchyMenuContext = Context.FindContext<UHierarchyMenuContext>();
+	if(HierarchyMenuContext->MenuHierarchyElements.Num() == 1)
+	{
+		HierarchyMenuContext->MenuHierarchyElements[0]->RequestRename();
+	}
+}
+
+bool UDataHierarchyViewModelBase::ToolMenuCanRequestRename(const FToolMenuContext& Context) const
+{
+	UHierarchyMenuContext* HierarchyMenuContext = Context.FindContext<UHierarchyMenuContext>();
+	if(HierarchyMenuContext->MenuHierarchyElements.Num() == 1)
+	{
+		return HierarchyMenuContext->MenuHierarchyElements[0]->CanRename();
+	}
+
+	return false;
+}
+
+void UDataHierarchyViewModelBase::ToolMenuDelete(const FToolMenuContext& Context) const
+{
+	UHierarchyMenuContext* HierarchyMenuContext = Context.FindContext<UHierarchyMenuContext>();
+	DeleteElements(HierarchyMenuContext->MenuHierarchyElements);
+}
+
+bool UDataHierarchyViewModelBase::ToolMenuCanDelete(const FToolMenuContext& Context) const
+{
+	UHierarchyMenuContext* HierarchyMenuContext = Context.FindContext<UHierarchyMenuContext>();
+	
+	for(TSharedPtr<FHierarchyElementViewModel> MenuHierarchyElement : HierarchyMenuContext->MenuHierarchyElements)
+	{
+		if(MenuHierarchyElement->CanDelete() == false)
+		{
+			return false;
+		}
+	}
+
+	return HierarchyMenuContext->MenuHierarchyElements.Num() > 0;
+}
+
+void UDataHierarchyViewModelBase::ToolMenuNavigateTo(const FToolMenuContext& Context) const
+{
+	UHierarchyMenuContext* HierarchyMenuContext = Context.FindContext<UHierarchyMenuContext>();
+
+	if(HierarchyMenuContext->MenuHierarchyElements.Num() == 1)
+	{
+		if(TSharedPtr<FHierarchyElementViewModel> MatchingViewModelInHierarchy = GetHierarchyRootViewModel()->FindViewModelForChild(HierarchyMenuContext->MenuHierarchyElements[0]->GetData()->GetPersistentIdentity(), true))
+		{
+			NavigateToElementInHierarchy(MatchingViewModelInHierarchy.ToSharedRef());
+		}
+	}
+}
+
+bool UDataHierarchyViewModelBase::ToolMenuCanNavigateTo(const FToolMenuContext& Context) const
+{
+	UHierarchyMenuContext* HierarchyMenuContext = Context.FindContext<UHierarchyMenuContext>();
+
+	if(HierarchyMenuContext->MenuHierarchyElements.Num() != 1)
+	{
+		return false;
+	}
+	
+	TSharedPtr<FHierarchyElementViewModel> ViewModel = HierarchyMenuContext->MenuHierarchyElements[0];
+	if(ViewModel->IsForHierarchy())
+	{
+		return false;
+	}
+	
+	if(TSharedPtr<FHierarchyElementViewModel> MatchingViewModelInHierarchy = GetHierarchyRootViewModel()->FindViewModelForChild(ViewModel->GetData()->GetPersistentIdentity(), true))
+	{
+		return MatchingViewModelInHierarchy.IsValid();
+	}
+	
+	return false;
+}
+
+const TArray<TSharedPtr<FHierarchyElementViewModel>>& UDataHierarchyViewModelBase::GetHierarchyItems() const
+{
+	return HierarchyRootViewModel->GetFilteredChildren();
+}
+
+FName UDataHierarchyViewModelBase::GetContextMenuName() const
+{
+	return FName(FString(TEXT("HierarchyEditor.") + GetClass()->GetName()));
+}
+
+TSharedRef<FHierarchyDragDropOp> UDataHierarchyViewModelBase::CreateDragDropOp(TSharedRef<FHierarchyElementViewModel> Item)
+{
+	TSharedRef<FHierarchyDragDropOp> DragDropOp = MakeShared<FHierarchyDragDropOp>(Item);
+	DragDropOp->Construct();
+	return DragDropOp;
+}
+
+void UDataHierarchyViewModelBase::OnGetChildren(TSharedPtr<FHierarchyElementViewModel> Element, TArray<TSharedPtr<FHierarchyElementViewModel>>& OutChildren) const
+{
+	OutChildren.Append(Element->GetFilteredChildren());
+}
+
+void UDataHierarchyViewModelBase::SetActiveHierarchySection(TSharedPtr<FHierarchySectionViewModel> Section)
+{
+	ActiveHierarchySection = Section;	
+	RefreshHierarchyView(true);
+	OnHierarchySectionActivatedDelegate.ExecuteIfBound(Section);
+}
+
+TSharedPtr<FHierarchySectionViewModel> UDataHierarchyViewModelBase::GetActiveHierarchySection() const
+{
+	return ActiveHierarchySection.IsValid() ? ActiveHierarchySection.Pin() : nullptr;
+}
+
+UHierarchySection* UDataHierarchyViewModelBase::GetActiveHierarchySectionData() const
+{
+	return ActiveHierarchySection.IsValid() ? ActiveHierarchySection.Pin()->GetDataMutable<UHierarchySection>() : nullptr;
+}
+
+bool UDataHierarchyViewModelBase::IsHierarchySectionActive(const UHierarchySection* Section) const
+{
+	return ActiveHierarchySection == nullptr || ActiveHierarchySection.Pin()->GetData() == Section;
+}
+
+FString UDataHierarchyViewModelBase::OnElementToStringDebug(TSharedPtr<FHierarchyElementViewModel> ElementViewModel) const
+{
+	return ElementViewModel->ToString();	
+}
+
+FHierarchyElementViewModel::~FHierarchyElementViewModel()
+{
+	Children.Empty();
+	FilteredChildren.Empty();
+}
+
+UHierarchyElement* FHierarchyElementViewModel::AddChild(TSubclassOf<UHierarchyElement> NewChildClass, FHierarchyElementIdentity ChildIdentity)
+{
+	UHierarchyElement* NewChild = NewObject<UHierarchyElement>(GetDataMutable(), NewChildClass);
+	NewChild->SetFlags(RF_Transactional);
+	NewChild->Modify();
+	NewChild->SetIdentity(ChildIdentity);
+	GetDataMutable()->GetChildrenMutable().Add(NewChild);
+	
+	SyncViewModelsToData();
+	HierarchyViewModel->OnHierarchyChanged().Broadcast();
+	return NewChild;
+}
+
+void FHierarchyElementViewModel::Tick(float DeltaTime)
+{
+	if(bRenamePending)
+	{
+		RequestRename();
+	}
+}
+
+TStatId FHierarchyElementViewModel::GetStatId() const
+{
+	RETURN_QUICK_DECLARE_CYCLE_STAT(FHierarchyElementViewModel, STATGROUP_Tickables);
+}
+
+void FHierarchyElementViewModel::RefreshChildrenData()
+{
+	TArray<TSharedPtr<FHierarchyElementViewModel>> TmpChildren = Children;
+	for(TSharedPtr<FHierarchyElementViewModel> Child : TmpChildren)
+	{
+		if(Child->RepresentsExternalData() && Child->DoesExternalDataStillExist(HierarchyViewModel->GetRefreshContext()) == false)
+		{
+			UE_LOGFMT(LogDataHierarchyEditor, Verbose, "Hierarchy Element {ElementName} no longer has valid external data. Deleting.", Child->ToString());
+			Child->Delete();
+		}
+	}
+
+	/** Every item view model can define its own sort order for its children. */
+	SortChildrenData();
+	
+	RefreshChildrenDataInternal();
+
+	/** All remaining children are supposed to exist at this point, as internal data won't be removed by refreshing & external data was cleaned up already.
+	 * This will not call RefreshChildrenData on data that has just been added as no view models exist for these yet.
+	 */
+	for(TSharedPtr<FHierarchyElementViewModel> Child : Children)
+	{
+		Child->RefreshChildrenData();
+	}
+}
+
+void FHierarchyElementViewModel::SyncViewModelsToData()
+{	
+	// this will recursively remove all outdated external data as well as give individual view models the chance to add new data
+	RefreshChildrenData();
+	
+	// now that the data is refreshed, we can sync to the data by recycling view models & creating new ones
+	// old view models will get deleted automatically
+	TArray<TSharedPtr<FHierarchyElementViewModel>> NewChildren;
+	for(UHierarchyElement* Child : Element->GetChildren())
+	{		
+		int32 ViewModelIndex = FindIndexOfChild(Child);
+		// if we couldn't find a view model for a data child, we create it here
+		if(ViewModelIndex == INDEX_NONE)
+		{
+			TSharedPtr<FHierarchyElementViewModel> ChildViewModel = HierarchyViewModel->CreateViewModelForElement(Child, AsShared());
+			if(ensure(ChildViewModel.IsValid()))
+			{
+				ChildViewModel->Initialize();
+				ChildViewModel->SyncViewModelsToData();
+				NewChildren.Add(ChildViewModel);
+			}
+		}
+		// if we could find the view model, we refresh its contained view models and readd it
+		else
+		{
+			Children[ViewModelIndex]->SyncViewModelsToData();
+			NewChildren.Add(Children[ViewModelIndex]);
+		}
+	}
+
+	Children.Empty();
+	Children.Append(NewChildren);
+	
+	for(TSharedPtr<FHierarchyElementViewModel> Child : Children)
+	{
+		Child->OnChildRequestedDeletion().BindSP(this, &FHierarchyElementViewModel::DeleteChild);
+		Child->GetOnSynced().BindSP(this, &FHierarchyElementViewModel::PropagateOnChildSynced);
+	}
+
+	/** Give the view models a chance to further customize the children sync process. */
+	SyncViewModelsToDataInternal();	
+
+	// then we sort the view models according to the data order as this is what will determine widget order created from the view models
+	Children.Sort([this](const TSharedPtr<FHierarchyElementViewModel>& ItemA, const TSharedPtr<FHierarchyElementViewModel>& ItemB)
+		{
+			return FindIndexOfDataChild(ItemA) < FindIndexOfDataChild(ItemB);
+		});
+	
+	// we refresh the filtered children here as well
+	GetFilteredChildren();
+
+	OnSyncedDelegate.ExecuteIfBound();
+}
+
+const TArray<TSharedPtr<FHierarchyElementViewModel>>& FHierarchyElementViewModel::GetFilteredChildren() const
+{
+	FilteredChildren.Empty();
+
+	if(CanHaveChildren())
+	{
+		for(TSharedPtr<FHierarchyElementViewModel> Child : Children)
+		{
+			bool bPassesFilter = true;
+			for(const FOnFilterChild& OnFilterChild : ChildFilters)
+			{
+				bPassesFilter &= OnFilterChild.Execute(Child);
+
+				if(!bPassesFilter)
+				{
+					break;
+				}
+			}
+
+			if(bPassesFilter)
+			{
+				FilteredChildren.Add(Child);
+			}
+		}
+	}
+
+	return FilteredChildren;
+}
+
+void FHierarchyElementViewModel::SortChildrenData() const
+{
+	GetDataMutable()->GetChildrenMutable().StableSort([](const UHierarchyElement& ItemA, const UHierarchyElement& ItemB)
+		{
+			return ItemA.IsA<UHierarchyCategory>() && ItemB.IsA<UHierarchyItem>();
+		});
+}
+
+int32 FHierarchyElementViewModel::GetHierarchyDepth() const
+{
+	if(Parent.IsValid())
+	{
+		return 1 + Parent.Pin()->GetHierarchyDepth();
+	}
+
+	return 0;
+}
+
+void FHierarchyElementViewModel::AddChildFilter(FOnFilterChild InFilterChild)
+{
+	if(ensure(InFilterChild.IsBound()))
+	{
+		ChildFilters.Add(InFilterChild);
+	}
+}
+
+bool FHierarchyElementViewModel::HasParent(TSharedPtr<FHierarchyElementViewModel> ParentCandidate, bool bRecursive) const
+{
+	if(Parent.IsValid())
+	{
+		if(Parent == ParentCandidate)
+		{
+			return true;
+		}
+		else if(bRecursive)
+		{
+			return Parent.Pin()->HasParent(ParentCandidate, bRecursive);
+		}
+	}
+
+	return false;
+}
+
+TSharedRef<FHierarchyElementViewModel> FHierarchyElementViewModel::DuplicateToThis(TSharedPtr<FHierarchyElementViewModel> ItemToDuplicate, int32 InsertIndex)
+{
+	UHierarchyElement* NewItem = Cast<UHierarchyElement>(StaticDuplicateObject(ItemToDuplicate->GetData(), GetDataMutable()));
+	if(InsertIndex == INDEX_NONE)
+	{
+		GetDataMutable()->GetChildrenMutable().Add(NewItem);
+	}
+	else
+	{
+		GetDataMutable()->GetChildrenMutable().Insert(NewItem, InsertIndex);
+	}
+	
+	SyncViewModelsToData();
+
+	HierarchyViewModel->OnHierarchyChanged().Broadcast();
+	TSharedPtr<FHierarchyElementViewModel> ViewModel = FindViewModelForChild(NewItem);
+	return ViewModel.ToSharedRef();
+}
+
+TSharedRef<FHierarchyElementViewModel> FHierarchyElementViewModel::ReparentToThis(TSharedPtr<FHierarchyElementViewModel> ItemToMove, int32 InsertIndex)
+{
+	UHierarchyElement* NewItem = Cast<UHierarchyElement>(StaticDuplicateObject(ItemToMove->GetData(), GetDataMutable()));
+	if(InsertIndex == INDEX_NONE)
+	{
+		GetDataMutable()->GetChildrenMutable().Add(NewItem);
+	}
+	else
+	{
+		GetDataMutable()->GetChildrenMutable().Insert(NewItem, InsertIndex);
+	}
+	
+	ItemToMove->Delete();
+	SyncViewModelsToData();
+	HierarchyViewModel->OnHierarchyChanged().Broadcast();
+	TSharedPtr<FHierarchyElementViewModel> ViewModel = FindViewModelForChild(NewItem);
+	return ViewModel.ToSharedRef();
+}
+
+TSharedPtr<FHierarchyElementViewModel> FHierarchyElementViewModel::FindViewModelForChild(UHierarchyElement* Child, bool bSearchRecursively) const
+{
+	int32 Index = FindIndexOfChild(Child);
+	if(Index != INDEX_NONE)
+	{
+		return Children[Index];
+	}
+
+	if(bSearchRecursively)
+	{
+		for(TSharedPtr<FHierarchyElementViewModel> ChildViewModel : Children)
+		{
+			TSharedPtr<FHierarchyElementViewModel> FoundViewModel = ChildViewModel->FindViewModelForChild(Child, bSearchRecursively);
+
+			if(FoundViewModel.IsValid())
+			{
+				return FoundViewModel;
+			}
+		}
+	}
+
+	return nullptr;
+}
+
+TSharedPtr<FHierarchyElementViewModel> FHierarchyElementViewModel::FindViewModelForChild(FHierarchyElementIdentity ChildIdentity, bool bSearchRecursively) const
+{
+	for(TSharedPtr<FHierarchyElementViewModel> Child : Children)
+	{
+		if(Child->GetData()->GetPersistentIdentity() == ChildIdentity)
+		{
+			return Child;
+		}
+	}
+
+	if(bSearchRecursively)
+	{
+		for(TSharedPtr<FHierarchyElementViewModel> ChildViewModel : Children)
+		{
+			TSharedPtr<FHierarchyElementViewModel> FoundViewModel = ChildViewModel->FindViewModelForChild(ChildIdentity, bSearchRecursively);
+
+			if(FoundViewModel.IsValid())
+			{
+				return FoundViewModel;
+			}
+		}
+	}
+
+	return nullptr;
+}
+
+int32 FHierarchyElementViewModel::FindIndexOfChild(UHierarchyElement* Child) const
+{
+	return Children.FindLastByPredicate([Child](TSharedPtr<FHierarchyElementViewModel> Item)
+	{
+		return Item->GetData() == Child;
+	});
+}
+
+int32 FHierarchyElementViewModel::FindIndexOfDataChild(TSharedPtr<FHierarchyElementViewModel> Child) const
+{
+	return GetData()->GetChildren().Find(Child->GetDataMutable());
+}
+
+int32 FHierarchyElementViewModel::FindIndexOfDataChild(UHierarchyElement* Child) const
+{
+	return GetData()->GetChildren().Find(Child);
+}
+
+void FHierarchyElementViewModel::Delete()
+{
+	OnChildRequestedDeletionDelegate.Execute(AsShared());
+}
+
+void FHierarchyElementViewModel::DeleteChild(TSharedPtr<FHierarchyElementViewModel> Child)
+{	
+	ensure(Child->GetParent().Pin() == AsShared());
+	GetDataMutable()->Modify();
+	GetDataMutable()->GetChildrenMutable().Remove(Child->GetDataMutable());
+	Children.Remove(Child);
+}
+
+TOptional<EItemDropZone> FHierarchyElementViewModel::OnCanRowAcceptDrop(const FDragDropEvent& DragDropEvent, EItemDropZone ItemDropZone, TSharedPtr<FHierarchyElementViewModel> Item)
+{
+	if(TSharedPtr<FHierarchyDragDropOp> DragDropOp = DragDropEvent.GetOperationAs<FHierarchyDragDropOp>())
+	{
+		FCanPerformActionResults Results = CanDropOn(DragDropOp->GetDraggedElement().Pin(), ItemDropZone);
+		DragDropOp->SetDescription(Results.CanPerformMessage);
+		return Results.bCanPerform ? ItemDropZone : TOptional<EItemDropZone>();
+	}
+
+	return TOptional<EItemDropZone>();
+}
+
+FReply FHierarchyElementViewModel::OnDroppedOnRow(const FDragDropEvent& DragDropEvent, EItemDropZone ItemDropZone, TSharedPtr<FHierarchyElementViewModel> Item)
+{
+	if(TSharedPtr<FHierarchyDragDropOp> HierarchyDragDropOp = DragDropEvent.GetOperationAs<FHierarchyDragDropOp>())
+	{
+		OnDroppedOn(HierarchyDragDropOp->GetDraggedElement().Pin(), ItemDropZone);
+		return FReply::Handled();
+	}
+
+	return FReply::Unhandled();
+}
+
+void FHierarchyElementViewModel::OnRowDragLeave(const FDragDropEvent& DragDropEvent)
+{
+	if(TSharedPtr<FHierarchyDragDropOp> HierarchyDragDropOp = DragDropEvent.GetOperationAs<FHierarchyDragDropOp>())
+	{
+		HierarchyDragDropOp->SetDescription(FText::GetEmpty());
+	}
+}
+
+FHierarchyElementViewModel::FCanPerformActionResults FHierarchyElementViewModel::CanDrag()
+{
+	FCanPerformActionResults Results = IsEditableByUser();
+	if(Results.bCanPerform == false)
+	{
+		return Results;
+	}
+
+	return CanDragInternal();
+}
+
+FHierarchyElementViewModel::FCanPerformActionResults FHierarchyElementViewModel::CanDropOnInternal(TSharedPtr<FHierarchyElementViewModel>, EItemDropZone ItemDropZone)
+{
+	return false;
+}
+
+void FHierarchyElementViewModel::PropagateOnChildSynced()
+{
+	OnSyncedDelegate.ExecuteIfBound();
+}
+
+FReply FHierarchyElementViewModel::OnDragDetected(const FGeometry& Geometry, const FPointerEvent& PointerEvent,	bool bIsSource)
+{
+	FCanPerformActionResults CanDragResults = CanDrag();
+	if(CanDragResults == true)
+	{
+		// if the drag is coming from source, we check if any of the hierarchy data already contains that element and we don't start a drag drop in that case
+		if(bIsSource)
+		{
+			TArray<TSharedPtr<FHierarchyElementViewModel>> AllChildren;
+			GetChildrenViewModelsForType<UHierarchyElement, FHierarchyElementViewModel>(AllChildren, true);
+
+			bool bCanDrag = GetHierarchyViewModel()->GetHierarchyRootViewModel()->FindViewModelForChild(GetData()->GetPersistentIdentity(), true) == nullptr;			
+
+			if(bCanDrag)
+			{
+				for(TSharedPtr<FHierarchyElementViewModel>& ChildViewModel : AllChildren)
+				{
+					if(GetHierarchyViewModel()->GetHierarchyRootViewModel()->FindViewModelForChild(ChildViewModel->GetData()->GetPersistentIdentity(), true) != nullptr)
+					{
+						bCanDrag = false;
+						break;
+					}
+				}
+			}
+			
+			if(bCanDrag == false)
+			{
+				return FReply::Unhandled();
+			}
+		}
+		
+		TSharedRef<FHierarchyDragDropOp> HierarchyDragDropOp = HierarchyViewModel->CreateDragDropOp(AsShared());
+		HierarchyDragDropOp->SetFromSourceList(bIsSource);
+
+		return FReply::Handled().BeginDragDrop(HierarchyDragDropOp);			
+	}
+	else
+	{
+		// if we can't drag and have a message, we show it as a slate notification
+		if(CanDragResults.CanPerformMessage.IsEmpty() == false)
+		{
+			FNotificationInfo CantDragInfo(CanDragResults.CanPerformMessage);
+			FSlateNotificationManager::Get().AddNotification(CantDragInfo);
+		}
+	}
+		
+	return FReply::Unhandled();
+}
+
+FHierarchyRootViewModel::~FHierarchyRootViewModel()
+{
+	
+}
+
+void FHierarchyRootViewModel::Initialize()
+{
+	GetOnSynced().BindSP(this, &FHierarchyRootViewModel::PropagateOnSynced);
+}
+
+FHierarchyElementViewModel::FCanPerformActionResults FHierarchyRootViewModel::CanDropOnInternal(TSharedPtr<FHierarchyElementViewModel> DraggedElement, EItemDropZone ItemDropZone)
+{
+	FCanPerformActionResults Results(false);
+	
+	// we only allow drops if some general conditions are fulfilled
+	if(DraggedElement->GetData() != GetData() &&
+		(!DraggedElement->HasParent(AsShared(), false) || ItemDropZone != EItemDropZone::OntoItem)  &&
+		!HasParent(DraggedElement, true))
+	{
+		Results.bCanPerform = 
+			// items can be dropped onto the root directly if the section is set to "All"
+			(DraggedElement->GetData()->IsA<UHierarchyItem>() && HierarchyViewModel->GetActiveHierarchySectionData() == nullptr)
+				||
+			// categories can be dropped onto the root always
+			(DraggedElement->GetData()->IsA<UHierarchyCategory>());
+
+		if(Results.bCanPerform)
+		{
+			if(DraggedElement->IsForHierarchy() == false)
+			{
+				FText Message = LOCTEXT("CanDropSourceItemOnRootDragMessage", "Add {0} to the hierarchy root.");
+				Message = FText::FormatOrdered(Message, FText::FromString(DraggedElement->ToString()));
+				Results.CanPerformMessage = Message;
+			}
+			else
+			{
+				FText Message = LOCTEXT("CanDropHierarchyItemOnRootDragMessage", "Move {0} to the hierarchy root.");
+				Message = FText::FormatOrdered(Message, FText::FromString(DraggedElement->ToString()));
+				Results.CanPerformMessage = Message;
+			}			
+		}
+		else
+		{
+			FText Message = LOCTEXT("CantDropHierarchyItemOnRootDragMessage", "Can not add {0} here. Please add it to a category!");
+			Message = FText::FormatOrdered(Message, FText::FromString(DraggedElement->ToString()));
+			Results.CanPerformMessage = Message;
+		}
+	}
+	
+	return Results;
+}
+
+void FHierarchyRootViewModel::OnDroppedOnInternal(TSharedPtr<FHierarchyElementViewModel> DroppedItem, EItemDropZone ItemDropZone)
+{
+	FScopedTransaction Transaction(LOCTEXT("Transaction_OnDropOnRoot", "Dropped item on root"));
+	HierarchyViewModel->GetHierarchyRoot()->Modify();
+
+	if(DroppedItem->GetDataMutable()->IsA<UHierarchyItem>() || DroppedItem->GetDataMutable()->IsA<UHierarchyCategory>())
+	{
+		TSharedPtr<FHierarchyElementViewModel> NewViewModel;
+		// we duplicate the item if the dragged item is from source
+		if(DroppedItem->IsForHierarchy() == false)
+		{
+			NewViewModel = DuplicateToThis(DroppedItem);
+		}
+		else
+		{
+			NewViewModel = ReparentToThis(DroppedItem);
+		}
+
+		if(UHierarchyCategory* AsCategory = Cast<UHierarchyCategory>(NewViewModel->GetDataMutable()))
+		{
+			AsCategory->SetSection(HierarchyViewModel->GetActiveHierarchySectionData());
+		}
+
+		HierarchyViewModel->RefreshHierarchyView();
+	}
+}
+
+TSharedPtr<FHierarchySectionViewModel> FHierarchyRootViewModel::AddSection()
+{
+	FScopedTransaction ScopedTransaction(LOCTEXT("NewSectionAdded","Added Section"));
+	HierarchyViewModel->GetHierarchyRoot()->Modify();
+	
+	UHierarchySection* SectionData = GetDataMutable<UHierarchyRoot>()->AddSection(LOCTEXT("HierarchyEditorDefaultNewSectionName", "Section"), 0, HierarchyViewModel->GetSectionDataClass());
+	SectionData->Modify();
+	TSharedPtr<FHierarchyElementViewModel> ViewModel = HierarchyViewModel->CreateViewModelForElement(SectionData, StaticCastSharedRef<FHierarchyRootViewModel>(AsShared()));
+	TSharedPtr<FHierarchySectionViewModel> SectionViewModel = StaticCastSharedPtr<FHierarchySectionViewModel>(ViewModel);
+	if(!ensureMsgf(SectionViewModel.IsValid(), TEXT("Make sure that CreateViewModelForData creates a FHierarchySectionViewModel (or derived) for UHierarchySection elements")))
+	{
+		return nullptr;	
+	}
+	
+	SectionViewModels.Add(SectionViewModel);
+	SyncViewModelsToData();
+	HierarchyViewModel->SetActiveHierarchySection(SectionViewModel);
+
+	OnSectionAddedDelegate.ExecuteIfBound(SectionViewModel);
+	OnSectionsChangedDelegate.ExecuteIfBound();
+	return SectionViewModel;
+}
+
+void FHierarchyRootViewModel::DeleteSection(TSharedPtr<FHierarchyElementViewModel> InSectionViewModel)
+{
+	TSharedPtr<FHierarchySectionViewModel> SectionViewModel = StaticCastSharedPtr<FHierarchySectionViewModel>(InSectionViewModel);
+	GetDataMutable<UHierarchyRoot>()->GetSectionDataMutable().Remove(SectionViewModel->GetDataMutable<UHierarchySection>());
+	SectionViewModels.Remove(SectionViewModel);
+
+	OnSectionDeletedDelegate.ExecuteIfBound(SectionViewModel);
+	OnSectionsChangedDelegate.ExecuteIfBound();
+}
+
+void FHierarchyRootViewModel::PropagateOnSynced()
+{
+	OnSyncPropagatedDelegate.ExecuteIfBound();
+}
+
+void FHierarchyRootViewModel::SyncViewModelsToDataInternal()
+{
+	const UHierarchyRoot* RootData = GetData<UHierarchyRoot>();
+
+	TArray<TSharedPtr<FHierarchySectionViewModel>> NewSectionViewModels;
+	TArray<TSharedPtr<FHierarchySectionViewModel>> SectionViewModelsToDelete;
+	
+	for(TSharedPtr<FHierarchySectionViewModel> SectionViewModel : SectionViewModels)
+	{
+		if(!RootData->GetSectionData().Contains(SectionViewModel->GetData()))
+		{
+			SectionViewModelsToDelete.Add(SectionViewModel);
+		}
+	}
+
+	for (TSharedPtr<FHierarchySectionViewModel> SectionViewModel : SectionViewModelsToDelete)
+	{
+		SectionViewModel->Delete();
+	}
+	
+	for(UHierarchySection* Section : RootData->GetSectionData())
+	{
+		TSharedPtr<FHierarchySectionViewModel>* SectionViewModelPtr = SectionViewModels.FindByPredicate([Section](TSharedPtr<FHierarchySectionViewModel> SectionViewModel)
+		{
+			return SectionViewModel->GetData() == Section;
+		});
+
+		TSharedPtr<FHierarchySectionViewModel> SectionViewModel = nullptr;
+
+		if(SectionViewModelPtr)
+		{
+			SectionViewModel = *SectionViewModelPtr;
+		}
+
+		if(SectionViewModel == nullptr)
+		{
+			SectionViewModel = MakeShared<FHierarchySectionViewModel>(Section, StaticCastSharedRef<FHierarchyRootViewModel>(AsShared()), HierarchyViewModel);
+			SectionViewModel->SyncViewModelsToData();;
+		}
+		
+		NewSectionViewModels.Add(SectionViewModel);
+	}
+
+	SectionViewModels.Empty();
+	SectionViewModels.Append(NewSectionViewModels);
+
+	for(TSharedPtr<FHierarchySectionViewModel> SectionViewModel : SectionViewModels)
+	{
+		SectionViewModel->OnChildRequestedDeletion().BindSP(this, &FHierarchyRootViewModel::DeleteSection);
+	}
+
+	SectionViewModels.Sort([this](const TSharedPtr<FHierarchySectionViewModel>& ItemA, const TSharedPtr<FHierarchySectionViewModel>& ItemB)
+		{
+			return
+			GetDataMutable<UHierarchyRoot>()->GetSectionData().Find(Cast<UHierarchySection>(ItemA->GetDataMutable()))
+				<
+			GetDataMutable<UHierarchyRoot>()->GetSectionData().Find(Cast<UHierarchySection>(ItemB->GetDataMutable())); 
+		});
+}
+
+FHierarchyElementViewModel::FCanPerformActionResults FHierarchySectionViewModel::CanDropOnInternal(TSharedPtr<FHierarchyElementViewModel> DraggedElement, EItemDropZone ItemDropZone)
+{
+	FCanPerformActionResults Results(false);
+	// we don't allow dropping onto source sections and we don't specify a message as the sections aren't going to light up as valid drop targets
+	if(IsForHierarchy() == false)
+	{
+		return false;
+	}
+	
+	if(const UHierarchyCategory* Category = Cast<UHierarchyCategory>(DraggedElement->GetData()))
+	{
+		if(ItemDropZone == EItemDropZone::OntoItem)
+		{
+			FText Message = LOCTEXT("DropCategoryOnSectionDragMessage", "Add {0} to section {1}");
+			Message = FText::FormatOrdered(Message, FText::FromString(DraggedElement->ToString()), FText::FromString(ToString()));
+			
+			Results.CanPerformMessage = Message;
+			Results.bCanPerform = GetData() != Category->GetSection();
+		}
+	}
+	else if(UHierarchySection* Section = Cast<UHierarchySection>(DraggedElement->GetDataMutable()))
+	{
+		const bool bSameSection = GetData() == Section;
+
+		if(ItemDropZone == EItemDropZone::OntoItem && bSameSection == false)
+		{
+			FText Message = LOCTEXT("CantDropSectionOnSectionDragMessage", "Can't drop section on a section.");
+			Results.CanPerformMessage = Message;
+			Results.bCanPerform = false;
+			return Results;
+		}
+		
+		int32 DraggedElementIndex = GetHierarchyViewModel()->GetHierarchyRoot()->GetSectionIndex(Section);
+		int32 InsertionIndex = GetHierarchyViewModel()->GetHierarchyRoot()->GetSectionIndex(GetDataMutable<UHierarchySection>());
+		// we add 1 to the insertion index if it's below an item because we either want to insert at the current index to place the item above, or at current+1 for below
+		InsertionIndex += ItemDropZone == EItemDropZone::AboveItem ? -1 : 1;
+
+		Results.bCanPerform = !bSameSection && DraggedElementIndex != InsertionIndex;
+
+		if(Results.bCanPerform)
+		{
+			if(ItemDropZone == EItemDropZone::AboveItem)
+			{
+				FText Message = LOCTEXT("MoveSectionLeftDragMessage", "Move section {0} to the left");
+				Message = FText::FormatOrdered(Message, FText::FromString(DraggedElement->ToString()));
+				Results.CanPerformMessage = Message;
+			}
+			else if(ItemDropZone == EItemDropZone::BelowItem)
+			{
+				FText Message = LOCTEXT("MoveSectionRightDragMessage", "Move section {0} to the right");
+				Message = FText::FormatOrdered(Message, FText::FromString(DraggedElement->ToString()));
+				Results.CanPerformMessage = Message;
+			}
+		}
+	}
+	else if(UHierarchyItem* Item = Cast<UHierarchyItem>(DraggedElement->GetDataMutable()))
+	{
+		FText Message = LOCTEXT("CantDropItemOnSectionDragMessage", "Can't drop items onto sections. Please drag a category onto section {0}");
+		Message = FText::FormatOrdered(Message, FText::FromString(ToString()));
+		Results.bCanPerform = false;
+		Results.CanPerformMessage = Message;
+	}
+
+	return Results;
+}
+
+void FHierarchySectionViewModel::OnDroppedOnInternal(TSharedPtr<FHierarchyElementViewModel> DroppedItem, EItemDropZone ItemDropZone)
+{
+	if(DroppedItem->GetData()->IsA<UHierarchySection>())
+	{
+		UHierarchySection* DraggedSectionData = DroppedItem->GetDataMutable<UHierarchySection>();
+
+		int32 IndexOfThis = HierarchyViewModel->GetHierarchyRoot()->GetSectionData().Find(GetDataMutable<UHierarchySection>());
+		int32 DraggedSectionIndex = HierarchyViewModel->GetHierarchyRoot()->GetSectionData().Find(DraggedSectionData);
+
+		TArray<TObjectPtr<UHierarchySection>>& SectionData = HierarchyViewModel->GetHierarchyRoot()->GetSectionDataMutable();
+		int32 Count = SectionData.Num();
+
+		bool bDropSucceeded = false;
+		// above constitutes to the left here
+		if(ItemDropZone == EItemDropZone::AboveItem)
+		{
+			SectionData.RemoveAt(DraggedSectionIndex);
+			SectionData.Insert(DraggedSectionData, FMath::Max(IndexOfThis, 0));
+
+			bDropSucceeded = true;
+		}
+		else if(ItemDropZone == EItemDropZone::BelowItem)
+		{
+			SectionData.RemoveAt(DraggedSectionIndex);
+
+			if(IndexOfThis + 1 > SectionData.Num())
+			{
+				SectionData.Add(DraggedSectionData);
+			}
+			else
+			{
+				SectionData.Insert(DraggedSectionData, FMath::Min(IndexOfThis+1, Count));
+			}
+
+			bDropSucceeded = true;
+
+		}
+
+		if(bDropSucceeded)
+		{
+			HierarchyViewModel->ForceFullRefresh();
+			HierarchyViewModel->OnHierarchyChanged().Broadcast();
+		}
+	}
+	else
+	{		
+		if(UHierarchyCategory* HierarchyCategory = DroppedItem->GetDataMutable<UHierarchyCategory>())
+		{
+			FScopedTransaction Transaction(LOCTEXT("Transaction_OnSectionDrop", "Moved category to section"));
+			HierarchyViewModel->GetHierarchyRoot()->Modify();
+			
+			HierarchyCategory->SetSection(GetDataMutable<UHierarchySection>());
+
+			// we null out any sections for all contained categories
+			TArray<UHierarchyCategory*> AllChildCategories;
+			HierarchyCategory->GetChildrenOfType<UHierarchyCategory>(AllChildCategories, true);
+			for(UHierarchyCategory* ChildCategory : AllChildCategories)
+			{
+				ChildCategory->SetSection(nullptr);
+			}
+
+			// we only need to reparent if the parent isn't already the root. This stops unnecessary reordering
+			if(DroppedItem->GetParent() != HierarchyViewModel->GetHierarchyRootViewModel())
+			{
+				HierarchyViewModel->GetHierarchyRootViewModel()->ReparentToThis(DroppedItem);
+			}
+			
+			HierarchyViewModel->RefreshHierarchyView();
+			HierarchyViewModel->OnHierarchyChanged().Broadcast();
+		}
+	}
+}
+
+void FHierarchySectionViewModel::FinalizeInternal()
+{
+	if(HierarchyViewModel->GetActiveHierarchySection() == AsShared())
+	{
+		HierarchyViewModel->SetActiveHierarchySection(nullptr);
+	}
+
+	// we make sure to reset all categories' section entry that were referencing this section
+	TArray<UHierarchyCategory*> AllCategories;
+	HierarchyViewModel->GetHierarchyRoot()->GetChildrenOfType<UHierarchyCategory>(AllCategories, true);
+
+	for(UHierarchyCategory* Category : AllCategories)
+	{
+		if(Category->GetSection() == GetData())
+		{
+			Category->SetSection(nullptr);
+		}
+	}
+}
+
+::FHierarchyElementViewModel::FCanPerformActionResults FHierarchyItemViewModel::CanDropOnInternal(TSharedPtr<FHierarchyElementViewModel> DraggedElement, EItemDropZone ItemDropZone)
+{
+	bool bAllowDrop = false;
+	
+	TSharedPtr<FHierarchyElementViewModel> SourceDropItem = DraggedElement;
+	TSharedPtr<FHierarchyElementViewModel> TargetDropItem = AsShared();
+
+	// we only allow drops if some general conditions are fulfilled
+	if(SourceDropItem->GetData() != TargetDropItem->GetData() &&
+		(!SourceDropItem->HasParent(TargetDropItem, false) || ItemDropZone != EItemDropZone::OntoItem)  &&
+		!TargetDropItem->HasParent(SourceDropItem, true))
+	{
+		// items can be generally be dropped above/below other items
+		bAllowDrop = (SourceDropItem->GetData()->IsA<UHierarchyItem>() && ItemDropZone != EItemDropZone::OntoItem);
+	}
+
+	return bAllowDrop;
+}
+
+void FHierarchyItemViewModel::OnDroppedOnInternal(TSharedPtr<FHierarchyElementViewModel> DroppedItem, EItemDropZone ItemDropZone)
+{
+	FScopedTransaction Transaction(LOCTEXT("Transaction_MovedItem", "Moved an item in the hierarchy"));
+	HierarchyViewModel->GetHierarchyRoot()->Modify();
+	
+	bool bDropSucceeded = false;
+	if(ItemDropZone == EItemDropZone::AboveItem)
+	{
+		int32 IndexOfThis = Parent.Pin()->FindIndexOfDataChild(AsShared());
+
+		if(DroppedItem->IsForHierarchy() == false)
+		{
+			Parent.Pin()->DuplicateToThis(DroppedItem, FMath::Max(IndexOfThis, 0));
+		}
+		else
+		{
+			Parent.Pin()->ReparentToThis(DroppedItem, FMath::Max(IndexOfThis, 0));
+		}
+
+		bDropSucceeded = true;
+	}
+	else if(ItemDropZone == EItemDropZone::BelowItem)
+	{
+		int32 IndexOfThis = Parent.Pin()->FindIndexOfDataChild(AsShared());
+		
+		if(DroppedItem->IsForHierarchy() == false)
+		{
+			Parent.Pin()->DuplicateToThis(DroppedItem, FMath::Min(IndexOfThis+1, Parent.Pin()->GetChildren().Num()));
+		}
+		else
+		{
+			Parent.Pin()->ReparentToThis(DroppedItem, FMath::Min(IndexOfThis+1, Parent.Pin()->GetChildren().Num()));
+		}
+
+		bDropSucceeded = true;
+	}
+
+	if(bDropSucceeded)
+	{
+		HierarchyViewModel->RefreshHierarchyView();
+		HierarchyViewModel->RefreshSourceView();
+	}
+	else
+	{
+		Transaction.Cancel();
+	}
+}
+
+void FHierarchyCategoryViewModel::OnDroppedOnInternal(TSharedPtr<FHierarchyElementViewModel> DroppedItem, EItemDropZone ItemDropZone)
+{
+	FScopedTransaction Transaction(LOCTEXT("Transaction_OnCategoryDrop", "Dropped item on/above/below category"));
+	HierarchyViewModel->GetHierarchyRoot()->Modify();
+	
+	if(UHierarchyCategory* Category = DroppedItem->GetDataMutable<UHierarchyCategory>())
+	{
+		// if we are dragging a category above/below another category and the new parent is going to be the root, we update its section to the active section
+		if(ItemDropZone != EItemDropZone::OntoItem)
+		{
+			if(Parent.IsValid() && Parent == HierarchyViewModel->GetHierarchyRootViewModel())
+			{
+				Category->SetSection(HierarchyViewModel->GetActiveHierarchySectionData());
+
+				// we null out any sections for all contained categories
+				TArray<UHierarchyCategory*> AllChildCategories;
+				Category->GetChildrenOfType<UHierarchyCategory>(AllChildCategories, true);
+				for(UHierarchyCategory* ChildCategory : AllChildCategories)
+				{
+					ChildCategory->SetSection(nullptr);
+				}
+			}				
+		}
+		// if we are dragging a category onto another category, we null out its section instead
+		else
+		{
+			Category->SetSection(nullptr);
+
+			// we null out any sections for all contained categories
+			TArray<UHierarchyCategory*> AllChildCategories;
+			Category->GetChildrenOfType<UHierarchyCategory>(AllChildCategories, true);
+			for(UHierarchyCategory* ChildCategory : AllChildCategories)
+			{
+				ChildCategory->SetSection(nullptr);
+			}
+		}			
+	}
+
+	// the actual moving of the item happens here
+	if(ItemDropZone == EItemDropZone::OntoItem)
+	{
+		if(DroppedItem->IsForHierarchy() == false)
+		{
+			DuplicateToThis(DroppedItem);
+		}
+		else
+		{
+			ReparentToThis(DroppedItem);
+		}
+	}
+	else if(ItemDropZone == EItemDropZone::AboveItem)
+	{
+		int32 IndexOfThis = Parent.Pin()->FindIndexOfDataChild(AsShared());
+		if(DroppedItem->IsForHierarchy() == false)
+		{
+			Parent.Pin()->DuplicateToThis(DroppedItem, FMath::Max(IndexOfThis, 0));
+		}
+		else
+		{				
+			Parent.Pin()->ReparentToThis(DroppedItem, FMath::Max(IndexOfThis, 0));
+		}
+	}
+	else if(ItemDropZone == EItemDropZone::BelowItem)
+	{
+		int32 IndexOfThis = Parent.Pin()->FindIndexOfDataChild(AsShared());
+		if(DroppedItem->IsForHierarchy() == false)
+		{
+			Parent.Pin()->DuplicateToThis(DroppedItem, FMath::Min(IndexOfThis+1, Parent.Pin()->GetChildren().Num()));
+		}
+		else
+		{
+			Parent.Pin()->ReparentToThis(DroppedItem, FMath::Min(IndexOfThis+1, Parent.Pin()->GetChildren().Num()));
+		}
+	}
+}
+
+void UDataHierarchyViewModelBase::AddCategory(TSharedPtr<FHierarchyElementViewModel> CategoryParent) const
+{
+	// If no category parent was specified, we add it to the root
+	if(CategoryParent == nullptr)
+	{
+		CategoryParent = GetHierarchyRootViewModel();	
+	}
+	
+	int32 HierarchyDepth = CategoryParent->GetHierarchyDepth();
+	if(HierarchyDepth > 15)
+	{
+		FNotificationInfo Info(LOCTEXT("TooManyNestedCategoriesToastText", "We currently only allow a hierarchy depth of 15."));
+		Info.ExpireDuration = 4.f;
+		FSlateNotificationManager::Get().AddNotification(Info);
+		return;
+	}
+
+	FText TransactionText = FText::FormatOrdered(LOCTEXT("Transaction_AddedItem", "Added new {0} to hierarchy"), FText::FromString(GetCategoryDataClass()->GetName()));
+	FScopedTransaction Transaction(TransactionText);
+	GetHierarchyRoot()->Modify();
+
+	UClass* CategoryClass = GetCategoryDataClass();
+	UHierarchyCategory* Category = Cast<UHierarchyCategory>(CategoryParent->AddChild(CategoryClass, FHierarchyElementIdentity()));
+	
+	TSharedPtr<FHierarchyElementViewModel> ViewModel = CategoryParent->FindViewModelForChild(Category, false);
+	if(ensureMsgf(ViewModel.IsValid(), TEXT("Could not find view model for new category of type '%s'. Please ensure your 'CreateViewModelForData' function creates a view model."), *CategoryClass->GetName()))
+	{		
+		TArray<UHierarchyCategory*> SiblingCategories;
+		Category->GetTypedOuter<UHierarchyElement>()->GetChildrenOfType<UHierarchyCategory>(SiblingCategories);
+		
+		TSet<FName> CategoryNames;
+		for(const auto& SiblingCategory : SiblingCategories)
+		{
+			CategoryNames.Add(SiblingCategory->GetCategoryName());
+		}
+
+		Category->SetCategoryName(UE::DataHierarchyEditor::GetUniqueName(FName("New Category"), CategoryNames));
+		
+		// we only set the section property if the current section isn't set to "All"
+		Category->SetSection(GetActiveHierarchySectionData());
+		
+		RefreshHierarchyView();
+
+		OnElementAddedDelegate.ExecuteIfBound(ViewModel);
+	}
+}
+
+void UDataHierarchyViewModelBase::AddSection() const
+{
+	TSharedPtr<FHierarchySectionViewModel> SectionViewModel = GetHierarchyRootViewModel()->AddSection();	
+	OnElementAddedDelegate.ExecuteIfBound(SectionViewModel);
+	OnHierarchyChangedDelegate.Broadcast();
+}
+
+void UDataHierarchyViewModelBase::GenerateDynamicContextMenu(UToolMenu* ToolMenu)
+{
+	UHierarchyMenuContext* HierarchyMenuContext = ToolMenu->FindContext<UHierarchyMenuContext>();
+
+	if(HierarchyMenuContext == nullptr || HierarchyMenuContext->HierarchyViewModel.IsValid() == false)
+	{
+		return;
+	}
+
+	UDataHierarchyViewModelBase* HierarchyViewModel = HierarchyMenuContext->HierarchyViewModel.Get();
+	HierarchyViewModel->GenerateDynamicContextMenuInternal(ToolMenu);
+	
+	if(HierarchyMenuContext->MenuHierarchyElements.Num() == 1)
+	{
+		HierarchyMenuContext->MenuHierarchyElements[0]->AppendDynamicContextMenuForSingleElement(ToolMenu);
+	}
+}
+
+void UDataHierarchyViewModelBase::GenerateDynamicContextMenuInternal(UToolMenu* DynamicToolMenu) const
+{
+	UHierarchyMenuContext* HierarchyMenuContext = DynamicToolMenu->FindContext<UHierarchyMenuContext>();
+
+	if(HierarchyMenuContext == nullptr || HierarchyMenuContext->HierarchyViewModel.IsValid() == false)
+	{
+		return;
+	}
+
+	UDataHierarchyViewModelBase* HierarchyViewModel = HierarchyMenuContext->HierarchyViewModel.Get();
+
+	DynamicToolMenu->AddMenuEntry("Dynamic", FToolMenuEntry::InitMenuEntryWithCommandList(FDataHierarchyEditorCommands::Get().FindInHierarchy, HierarchyViewModel->GetCommands(), TAttribute<FText>(), TAttribute<FText>(), FSlateIcon(FAppStyle::GetAppStyleSetName(), "Icons.Find")));
+	DynamicToolMenu->AddMenuEntry("Dynamic", FToolMenuEntry::InitMenuEntryWithCommandList(FGenericCommands::Get().Rename, HierarchyViewModel->GetCommands()));
+	DynamicToolMenu->AddMenuEntry("Dynamic", FToolMenuEntry::InitMenuEntryWithCommandList(FGenericCommands::Get().Delete, HierarchyViewModel->GetCommands()));
+}
+
+UHierarchyElement* UDataHierarchyViewModelBase::AddElementUnderRoot(TSubclassOf<UHierarchyElement> NewChildClass, FHierarchyElementIdentity ChildIdentity)
+{
+	FScopedTransaction Transaction(LOCTEXT("Transaction_AddItem", "Add hierarchy item"));
+	HierarchyRoot->Modify();
+	return GetHierarchyRootViewModel()->AddChild(NewChildClass, ChildIdentity);
+}
+
+void UDataHierarchyViewModelBase::DeleteElementWithIdentity(FHierarchyElementIdentity Identity)
+{
+	FScopedTransaction Transaction(LOCTEXT("Transaction_DeleteItem", "Deleted hierarchy item"));
+	HierarchyRoot->Modify();
+
+	bool bItemDeleted = false;
+	if(TSharedPtr<FHierarchyElementViewModel> ViewModel = HierarchyRootViewModel->FindViewModelForChild(Identity, true))
+	{
+		if(ViewModel->CanDelete())
+		{
+			ViewModel->Delete();
+			bItemDeleted = true;
+		}
+	}
+	
+	TArray<TSharedPtr<FHierarchySectionViewModel>> SectionViewModels = HierarchyRootViewModel->GetSectionViewModels();
+	for(TSharedPtr<FHierarchySectionViewModel> SectionViewModel : SectionViewModels)
+	{
+		if(SectionViewModel->GetData()->GetPersistentIdentity() == Identity && SectionViewModel->CanDelete())
+		{
+			SectionViewModel->Delete();
+			bItemDeleted = true;
+		}
+	}
+
+	if(bItemDeleted)
+	{
+		HierarchyRootViewModel->SyncViewModelsToData();
+		OnHierarchyChangedDelegate.Broadcast();
+	}
+}
+
+void UDataHierarchyViewModelBase::DeleteElementsWithIdentities(TArray<FHierarchyElementIdentity> Identities)
+{
+	TArray<TSharedPtr<FHierarchyElementViewModel>> ViewModelsToDelete;
+	
+	for(FHierarchyElementIdentity& Identity : Identities)
+	{
+		if(TSharedPtr<FHierarchyElementViewModel> ViewModel = HierarchyRootViewModel->FindViewModelForChild(Identity, true))
+		{
+			ViewModelsToDelete.Add(ViewModel);
+		}
+
+		TArray<TSharedPtr<FHierarchySectionViewModel>> SectionViewModels = HierarchyRootViewModel->GetSectionViewModels();
+		for(TSharedPtr<FHierarchySectionViewModel> SectionViewModel : SectionViewModels)
+		{
+			if(SectionViewModel->GetData()->GetPersistentIdentity() == Identity)
+			{
+				ViewModelsToDelete.Add(SectionViewModel);
+			}
+		}
+	}
+
+	DeleteElements(ViewModelsToDelete);
+}
+
+void UDataHierarchyViewModelBase::DeleteElements(TArray<TSharedPtr<FHierarchyElementViewModel>> ViewModels) const
+{
+	FScopedTransaction Transaction(LOCTEXT("Transaction_DeleteHierarchyElements", "Deleted hierarchy elements"));
+	HierarchyRoot->Modify();
+	
+	bool bAnyItemsDeleted = false;
+	for(TSharedPtr<FHierarchyElementViewModel> ViewModel : ViewModels)
+	{
+		if(ViewModel->CanDelete())
+		{
+			ViewModel->Delete();
+			bAnyItemsDeleted = true;
+		}
+	}
+
+	if(bAnyItemsDeleted)
+	{
+		HierarchyRootViewModel->SyncViewModelsToData();
+		OnHierarchyChangedDelegate.Broadcast();
+	}
+}
+
+void UDataHierarchyViewModelBase::NavigateToElementInHierarchy(const FHierarchyElementIdentity& HierarchyIdentity) const
+{
+	OnNavigateToElementIdentityInHierarchyRequestedDelegate.ExecuteIfBound(HierarchyIdentity);
+}
+
+void UDataHierarchyViewModelBase::NavigateToElementInHierarchy(const TSharedRef<FHierarchyElementViewModel> HierarchyElement) const
+{
+	OnNavigateToElementInHierarchyRequestedDelegate.ExecuteIfBound(HierarchyElement);
+}
+
+FHierarchyDragDropOp::FHierarchyDragDropOp(TSharedPtr<FHierarchyElementViewModel> InDraggedElementViewModel) : DraggedElement(InDraggedElementViewModel)
+{
+	SetLabel(DraggedElement.Pin()->ToStringAsText());
+}
+
+TSharedPtr<SWidget> FHierarchyDragDropOp::GetDefaultDecorator() const
+{
+	TSharedRef<SWidget> CustomDecorator = CreateCustomDecorator();
+
+	SVerticalBox::FSlot* CustomSlot;
+	TSharedPtr<SWidget> Decorator = SNew(SToolTip)
+	[
+		SNew(SVerticalBox)
+		+ SVerticalBox::Slot()
+		.Expose(CustomSlot).
+		AutoHeight()
+		+ SVerticalBox::Slot()
+		.AutoHeight()
+		.Padding(2.f)
+		[
+			SNew(STextBlock)
+			.Text(this, &FHierarchyDragDropOp::GetLabel)
+			.TextStyle(&FAppStyle::GetWidgetStyle<FTextBlockStyle>("NormalText.Important"))
+			.Visibility_Lambda([this, CustomDecorator]()
+			{
+				return GetLabel().IsEmpty() || CustomDecorator != SNullWidget::NullWidget ? EVisibility::Collapsed : EVisibility::Visible;
+			})
+		]
+		+ SVerticalBox::Slot()
+		.AutoHeight()
+		.Padding(2.f)
+		[
+			SNew(STextBlock)
+			.Text(this, &FHierarchyDragDropOp::GetDescription)
+			.TextStyle(&FAppStyle::GetWidgetStyle<FTextBlockStyle>("NormalText"))
+			.Visibility_Lambda([this]()
+			{
+				return GetDescription().IsEmpty() ? EVisibility::Collapsed : EVisibility::Visible;
+			})
+		]
+	];
+
+	if(CustomDecorator != SNullWidget::NullWidget)
+	{
+		CustomSlot->AttachWidget(CustomDecorator);
+	}
+
+	return Decorator;
+}
+
+TSharedRef<SWidget> FSectionDragDropOp::CreateCustomDecorator() const
+{
+	return SNew(SCheckBox)
+		.Visibility(EVisibility::HitTestInvisible)
+		.Style(FAppStyle::Get(), "DetailsView.SectionButton")
+		.IsChecked(ECheckBoxState::Unchecked)
+		[
+			SNew(SInlineEditableTextBlock)
+			.Text(GetDraggedSection().Pin()->GetSectionNameAsText())
+		];
+}
+
+#undef LOCTEXT_NAMESPACE
diff --git a/Engine/Source/Editor/DataHierarchyEditor/Private/Widgets/SDataHierarchyEditor.cpp b/Engine/Source/Editor/DataHierarchyEditor/Private/Widgets/SDataHierarchyEditor.cpp
new file mode 100644
--- /dev/null
+++ b/Engine/Source/Editor/DataHierarchyEditor/Private/Widgets/SDataHierarchyEditor.cpp
@@ -0,0 +1,1758 @@
+// Copyright Epic Games, Inc. All Rights Reserved.
+
+#include "Widgets/SDataHierarchyEditor.h"
+
+#include "DataHierarchyEditorCommands.h"
+#include "UObject/Package.h"
+#include "Framework/Commands/GenericCommands.h"
+#include "Styling/StyleColors.h"
+#include "Styling/AppStyle.h"
+#include "DataHierarchyViewModelBase.h"
+#include "Widgets/Layout/SSeparator.h"
+#include "Widgets/Layout/SWrapBox.h"
+#include "Widgets/Input/SCheckBox.h"
+#include "Widgets/Input/SButton.h"
+#include "Widgets/Images/SImage.h"
+#include "Widgets/Text/SInlineEditableTextBlock.h"
+#include "Widgets/Text/STextBlock.h"
+#include "ScopedTransaction.h"
+#include "Modules/ModuleManager.h"
+#include "PropertyEditorModule.h"
+#include "Framework/Application/MenuStack.h"
+#include "Framework/Application/SlateApplication.h"
+#include "ToolMenus.h"
+#include "DataHierarchyEditorStyle.h"
+#include "Widgets/Input/SSearchBox.h"
+
+#define LOCTEXT_NAMESPACE "DataHierarchyEditor"
+
+TSharedRef<SWidget> SummonContextMenu(TArray<TSharedPtr<FHierarchyElementViewModel>> MenuHierarchyElements)
+{
+	UHierarchyMenuContext* MenuContextObject = NewObject<UHierarchyMenuContext>();
+	MenuContextObject->MenuHierarchyElements = MenuHierarchyElements;
+	TWeakObjectPtr<UDataHierarchyViewModelBase> ViewModel = MenuHierarchyElements[0]->GetHierarchyViewModel();
+	MenuContextObject->HierarchyViewModel = ViewModel;
+
+	FToolMenuContext MenuContext(MenuContextObject);
+	MenuContext.AppendCommandList(ViewModel->GetCommands());
+	
+	TSharedRef<SWidget> MenuWidget = UToolMenus::Get()->GenerateWidget(ViewModel->GetContextMenuName(), MenuContext);
+	return MenuWidget;
+}
+
+void SHierarchyCategory::Construct(const FArguments& InArgs,	TSharedPtr<FHierarchyCategoryViewModel> InCategoryViewModel)
+{
+	CategoryViewModel = InCategoryViewModel;
+	
+	InCategoryViewModel->GetOnRequestRename().BindSP(this, &SHierarchyCategory::EnterEditingMode);
+
+	const UHierarchyCategory* CategoryData = InCategoryViewModel->GetData<UHierarchyCategory>();
+	SetToolTipText(TAttribute<FText>::CreateLambda([CategoryData]
+	{
+		return CategoryData->GetTooltip();
+	}));
+	
+	ChildSlot
+	[
+		SAssignNew(InlineEditableTextBlock, SInlineEditableTextBlock)
+		.Style(&FDataHierarchyEditorStyle::Get().GetWidgetStyle<FInlineEditableTextBlockStyle>("HierarchyEditor.CategoryTextStyle"))
+		.Text(this, &SHierarchyCategory::GetCategoryText)
+		.OnTextCommitted(this, &SHierarchyCategory::OnRenameCategory)
+		.OnVerifyTextChanged(this, &SHierarchyCategory::OnVerifyCategoryRename)
+		.IsSelected(InArgs._IsSelected)
+	];
+}
+
+void SHierarchyCategory::EnterEditingMode() const
+{
+	if(CategoryViewModel.Pin()->CanRename())
+	{
+		InlineEditableTextBlock->EnterEditingMode();
+	}
+}
+
+bool SHierarchyCategory::OnVerifyCategoryRename(const FText& NewName, FText& OutTooltip) const
+{
+	TArray<TSharedPtr<FHierarchyCategoryViewModel>> SiblingCategoryViewModels;
+	CategoryViewModel.Pin()->GetParent().Pin()->GetChildrenViewModelsForType<UHierarchyCategory, FHierarchyCategoryViewModel>(SiblingCategoryViewModels);
+
+	if(GetCategoryText().ToString() != NewName.ToString())
+	{
+		TSet<FString> CategoryNames;
+		for(const auto& SiblingCategoryViewModel : SiblingCategoryViewModels)
+		{
+			CategoryNames.Add(Cast<UHierarchyCategory>(SiblingCategoryViewModel->GetDataMutable())->GetCategoryName().ToString());
+		}
+
+		if(CategoryNames.Contains(NewName.ToString()))
+		{
+			OutTooltip = LOCTEXT("HierarchyCategoryCantRename_DuplicateOnLayer", "Another category of the same name already exists!");
+			return false;
+		}
+	}
+
+	return true;
+}
+
+FText SHierarchyCategory::GetCategoryText() const
+{
+	return FText::FromString(CategoryViewModel.Pin()->ToString());
+}
+
+void SHierarchyCategory::OnRenameCategory(const FText& NewText, ETextCommit::Type) const
+{
+	FScopedTransaction Transaction(LOCTEXT("Transaction_Rename_Category", "Renamed hierarchy category"));
+	CategoryViewModel.Pin()->GetHierarchyViewModel()->GetHierarchyRoot()->Modify();
+	
+	CategoryViewModel.Pin()->Rename(FName(NewText.ToString()));
+}
+
+void SHierarchySection::Construct(const FArguments& InArgs, TSharedPtr<FHierarchySectionViewModel> InSection, TWeakObjectPtr<UDataHierarchyViewModelBase> InHierarchyViewModel)
+{
+	SectionViewModel = InSection;
+	HierarchyViewModel = InHierarchyViewModel;
+
+	IsSectionActive = InArgs._IsSectionActive;
+	OnSectionActivatedDelegate = InArgs._OnSectionActivated;
+	bForbidDropOn = InArgs._bForbidDropOn;
+	
+	if(SectionViewModel != nullptr)
+	{
+		SectionViewModel->GetOnRequestRename().BindSP(this, &SHierarchySection::TryEnterEditingMode);
+
+		SDropTarget::FArguments LeftDropTargetArgs;
+		SDropTarget::FArguments OntoDropTargetArgs;
+		SDropTarget::FArguments RightDropTargetArgs;
+
+		LeftDropTargetArgs
+			.OnAllowDrop(this, &SHierarchySection::OnCanAcceptDrop, EItemDropZone::AboveItem)
+			.OnDropped(this, &SHierarchySection::OnDroppedOn, EItemDropZone::AboveItem)
+			.VerticalImage(FAppStyle::GetNoBrush())
+			.HorizontalImage(FDataHierarchyEditorStyle::Get().GetBrush("HierarchyEditor.Stack.DropTarget.BorderHorizontal"));
+
+		OntoDropTargetArgs
+			.OnAllowDrop(this, &SHierarchySection::OnCanAcceptDrop, EItemDropZone::OntoItem)
+			.OnDropped(this, &SHierarchySection::OnDroppedOn, EItemDropZone::OntoItem);
+		
+		RightDropTargetArgs
+			.OnAllowDrop(this, &SHierarchySection::OnCanAcceptDrop, EItemDropZone::BelowItem)
+			.OnDropped(this, &SHierarchySection::OnDroppedOn, EItemDropZone::BelowItem)
+			.VerticalImage(FAppStyle::GetNoBrush())
+			.HorizontalImage(FDataHierarchyEditorStyle::Get().GetBrush("HierarchyEditor.Stack.DropTarget.BorderHorizontal"));
+
+		SetToolTipText(TAttribute<FText>::CreateSP(this, &SHierarchySection::GetTooltipText));
+
+		ChildSlot
+		[
+			SNew(SHorizontalBox)
+			+ SHorizontalBox::Slot()
+			.AutoWidth()
+			[
+				SNew(SSectionDragDropTarget, SectionViewModel, EItemDropZone::AboveItem)
+				.DropTargetArgs(LeftDropTargetArgs)
+			]
+			+ SHorizontalBox::Slot()
+			[
+				SNew(SSectionDragDropTarget, SectionViewModel, EItemDropZone::OntoItem)
+				.DropTargetArgs(OntoDropTargetArgs
+				.Content()
+				[
+					SAssignNew(CheckBox, SCheckBox)
+					.Visibility(EVisibility::HitTestInvisible)
+					.Style(FAppStyle::Get(), "DetailsView.SectionButton")
+					.OnCheckStateChanged(this, &SHierarchySection::OnSectionCheckChanged)
+					.IsChecked(this, &SHierarchySection::GetSectionCheckState)
+					.Padding(FMargin(8.f, 4.f))
+					[
+						SNew(SHorizontalBox)
+						+ SHorizontalBox::Slot()
+						.AutoWidth()
+						.Padding(2.f)
+						[
+							SNew(SImage).Image(SectionViewModel->GetSectionImage() != nullptr ? SectionViewModel->GetSectionImage() : FStyleDefaults::GetNoBrush())
+						]
+						+ SHorizontalBox::Slot()
+						[
+							SAssignNew(InlineEditableTextBlock, SInlineEditableTextBlock)
+							.Visibility(EVisibility::HitTestInvisible)
+							.Text(this, &SHierarchySection::GetText)
+							.OnTextCommitted(this, &SHierarchySection::OnRenameSection)
+							.OnVerifyTextChanged(this, &SHierarchySection::OnVerifySectionRename)
+							.IsSelected(this, &SHierarchySection::IsSectionSelected)
+							.IsReadOnly(this, &SHierarchySection::IsSectionReadOnly)
+						]
+					]
+				])
+			]
+			+ SHorizontalBox::Slot()
+			.AutoWidth()
+			[
+				SNew(SSectionDragDropTarget, SectionViewModel, EItemDropZone::BelowItem)
+				.DropTargetArgs(RightDropTargetArgs)				
+			]
+		];		
+	}
+	// if this section doesn't represent data, it's the "All" widget
+	else
+	{
+		ChildSlot
+		[
+			SNew(SDropTarget)
+			.OnAllowDrop(this, &SHierarchySection::OnCanAcceptDrop, EItemDropZone::OntoItem)
+			.OnDropped(this, &SHierarchySection::OnDroppedOn, EItemDropZone::OntoItem)
+			[
+				SNew(SCheckBox)
+				.Style(FAppStyle::Get(), "DetailsView.SectionButton")
+				.OnCheckStateChanged(this, &SHierarchySection::OnSectionCheckChanged)
+				.IsChecked(this, &SHierarchySection::GetSectionCheckState)
+				.Padding(FMargin(8.f, 4.f))
+				[
+					SNew(STextBlock)
+					.Text(LOCTEXT("AllSection", "All"))
+				]
+			]			
+		];
+	}
+}
+
+SHierarchySection::~SHierarchySection()
+{
+	SectionViewModel.Reset();
+}
+
+void SHierarchySection::TryEnterEditingMode() const
+{
+	if(SectionViewModel.IsValid() && SectionViewModel->CanRename())
+	{
+		InlineEditableTextBlock->EnterEditingMode();
+	}
+}
+
+TSharedPtr<FHierarchySectionViewModel> SHierarchySection::GetSectionViewModel()
+{
+	return SectionViewModel;
+}
+
+bool SHierarchySection::OnCanAcceptDrop(TSharedPtr<FDragDropOperation> DragDropOperation, EItemDropZone ItemDropZone) const
+{
+	if(bForbidDropOn)
+	{
+		return false;
+	}
+
+	if(DragDropOperation->IsOfType<FHierarchyDragDropOp>())
+	{
+		TSharedPtr<FHierarchyDragDropOp> HierarchyDragDropOp = StaticCastSharedPtr<FHierarchyDragDropOp>(DragDropOperation);
+		TSharedPtr<FHierarchyElementViewModel> DraggedItem = HierarchyDragDropOp->GetDraggedElement().Pin();
+		
+		if(SectionViewModel.IsValid())
+		{
+			return SectionViewModel->CanDropOn(DraggedItem, ItemDropZone).bCanPerform;
+		}
+		
+		// for the All section which has no valid view model, we simply do a check if the sections of the dragged categories are different
+		if(DragDropOperation->IsOfType<FHierarchyDragDropOp>())
+		{
+			if(const UHierarchyCategory* Category = Cast<UHierarchyCategory>(DraggedItem->GetData()))
+			{
+				return TryGetSectionData() != Category->GetSection();
+			}
+		}
+	}
+
+	return false;
+}
+
+FReply SHierarchySection::OnDroppedOn(const FGeometry&, const FDragDropEvent& DragDropEvent, EItemDropZone DropZone) const
+{
+	bDraggedOn = false;
+	if(TSharedPtr<FHierarchyDragDropOp> HierarchyDragDropOp = DragDropEvent.GetOperationAs<FHierarchyDragDropOp>())
+	{
+		TSharedPtr<FHierarchyElementViewModel> DraggedItem = HierarchyDragDropOp->GetDraggedElement().Pin();
+		if(SectionViewModel.IsValid())
+		{
+			SectionViewModel->OnDroppedOn(DraggedItem, DropZone);
+			return FReply::Handled();
+		}
+		
+		if(UHierarchyCategory* Category = Cast<UHierarchyCategory>(DraggedItem->GetDataMutable()))
+		{
+			Category->SetSection(nullptr);
+			TArray<UHierarchyCategory*> ChildrenCategories;
+			Category->GetChildrenOfType<>(ChildrenCategories, true);
+
+			for(UHierarchyCategory* ChildCategory : ChildrenCategories)
+			{
+				ChildCategory->SetSection(nullptr);
+			}
+
+			// we only need to reparent if the parent isn't already the root. This stops unnecessary reordering
+			if(HierarchyDragDropOp->GetDraggedElement().Pin()->GetParent() != HierarchyViewModel->GetHierarchyRootViewModel())
+			{
+				HierarchyViewModel->GetHierarchyRootViewModel()->ReparentToThis(HierarchyDragDropOp->GetDraggedElement().Pin());
+			}
+		
+			HierarchyViewModel->RefreshHierarchyView();
+
+			return FReply::Handled();
+		}
+	}
+
+	return FReply::Unhandled();
+}
+
+FReply SHierarchySection::OnMouseButtonDown(const FGeometry& MyGeometry, const FPointerEvent& MouseEvent)
+{
+	if(SectionViewModel.IsValid())
+	{
+		// we handle the event here so we can react on mouse button up
+		if(MouseEvent.IsMouseButtonDown(EKeys::RightMouseButton))
+		{
+			return FReply::Handled();
+		}
+		else if(MouseEvent.IsMouseButtonDown(EKeys::LeftMouseButton))
+		{			
+			OnSectionActivatedDelegate.ExecuteIfBound(SectionViewModel);
+			return FReply::Handled().DetectDrag(AsShared(), EKeys::LeftMouseButton).SetUserFocus(AsShared());	
+		}
+	}
+
+	return FReply::Unhandled();
+}
+
+FReply SHierarchySection::OnMouseButtonUp(const FGeometry& MyGeometry, const FPointerEvent& MouseEvent)
+{
+	if(SectionViewModel.IsValid() && SectionViewModel->IsForHierarchy())
+	{
+		if(MouseEvent.GetEffectingButton() == EKeys::RightMouseButton)
+		{			
+			// Show menu to choose getter vs setter
+			FSlateApplication::Get().PushMenu(
+				AsShared(),
+				FWidgetPath(),
+				SummonContextMenu({GetSectionViewModel()}),
+				FSlateApplication::Get().GetCursorPos(),
+				FPopupTransitionEffect(FPopupTransitionEffect::ContextMenu)
+			);
+				
+			OnSectionActivatedDelegate.ExecuteIfBound(SectionViewModel);			
+			return FReply::Handled();
+		}
+	}
+
+	return FReply::Unhandled();
+}
+
+void SHierarchySection::OnDragEnter(const FGeometry& MyGeometry, const FDragDropEvent& DragDropEvent)
+{
+	bDraggedOn = true;
+	if(DragDropEvent.GetOperationAs<FHierarchyDragDropOp>() && DragDropEvent.GetOperationAs<FSectionDragDropOp>() == nullptr)
+	{
+		RegisterActiveTimer(1.0f, FWidgetActiveTimerDelegate::CreateSP(this, &SHierarchySection::ActivateSectionIfDragging));
+	}
+}
+
+FReply SHierarchySection::OnDragDetected(const FGeometry& MyGeometry, const FPointerEvent& MouseEvent)
+{
+	if(MouseEvent.IsMouseButtonDown(EKeys::LeftMouseButton) && SectionViewModel.IsValid() && SectionViewModel->CanDrag().bCanPerform)
+	{
+		TSharedRef<FSectionDragDropOp> SectionDragDropOp = MakeShared<FSectionDragDropOp>(SectionViewModel);
+		SectionDragDropOp->Construct();
+		return FReply::Handled().BeginDragDrop(SectionDragDropOp);
+	}
+	
+	return FReply::Unhandled();
+}
+
+void SHierarchySection::OnDragLeave(const FDragDropEvent& DragDropEvent)
+{
+	bDraggedOn = false;
+	if(TSharedPtr<FHierarchyDragDropOp> HierarchyDragDropOp = DragDropEvent.GetOperationAs<FHierarchyDragDropOp>())
+	{
+		HierarchyDragDropOp->SetDescription(FText::GetEmpty());
+	}
+}
+
+FReply SHierarchySection::OnKeyDown(const FGeometry& MyGeometry, const FKeyEvent& InKeyEvent)
+{
+	if(SectionViewModel.IsValid() && InKeyEvent.GetKey() == EKeys::Delete && SectionViewModel->CanDelete())
+	{
+		SectionViewModel->Delete();
+		return FReply::Handled();
+	}
+
+	return FReply::Unhandled();
+}
+
+void SHierarchySection::OnMouseLeave(const FPointerEvent& MouseEvent)
+{
+	//bDraggedOn = false;
+	SCompoundWidget::OnMouseLeave(MouseEvent);
+}
+
+UHierarchySection* SHierarchySection::TryGetSectionData() const
+{
+	return SectionViewModel.IsValid() ? Cast<UHierarchySection>(SectionViewModel->GetDataMutable()) : nullptr;
+}
+
+FText SHierarchySection::GetText() const
+{
+	return SectionViewModel->GetSectionNameAsText();
+}
+
+FText SHierarchySection::GetTooltipText() const
+{
+	return SectionViewModel->GetSectionTooltip();
+}
+
+void SHierarchySection::OnRenameSection(const FText& Text, ETextCommit::Type CommitType) const
+{
+	FScopedTransaction Transaction(LOCTEXT("Transaction_Rename_Section", "Renamed hierarchy section"));
+	HierarchyViewModel->GetHierarchyRoot()->Modify();
+	
+	SectionViewModel->Rename(FName(Text.ToString()));
+}
+
+bool SHierarchySection::OnVerifySectionRename(const FText& NewName, FText& OutTooltip) const
+{
+	// this function shouldn't be used in case the section isn't valid but we'll make sure regardless
+	if(!SectionViewModel.IsValid())
+	{
+		return false;
+	}
+
+	if(SectionViewModel->GetSectionName().ToString() != NewName.ToString())
+	{
+		TArray<FString> SectionNames;
+
+		SectionNames.Add("All");
+		for(auto& Section : HierarchyViewModel->GetHierarchyRootViewModel()->GetSectionViewModels())
+		{
+			SectionNames.Add(Section->GetSectionName().ToString());
+		}
+
+		if(SectionNames.Contains(NewName.ToString()))
+		{
+			OutTooltip = LOCTEXT("HierarchySectionCantRename_Duplicate", "A section with that name already exists!");
+			return false;
+		}
+	}
+
+	return true;
+}
+
+bool SHierarchySection::IsSectionSelected() const
+{
+	return GetSectionCheckState() == ECheckBoxState::Checked ? true : false;
+}
+
+bool SHierarchySection::IsSectionReadOnly() const
+{
+	return SectionViewModel.IsValid() ? !SectionViewModel->CanRename() : true;
+}
+
+ECheckBoxState SHierarchySection::GetSectionCheckState() const
+{
+	return IsSectionActive.Get();
+}
+
+void SHierarchySection::OnSectionCheckChanged(ECheckBoxState NewState)
+{
+	OnSectionActivatedDelegate.ExecuteIfBound(SectionViewModel);
+}
+
+EActiveTimerReturnType SHierarchySection::ActivateSectionIfDragging(double CurrentTime, float DeltaTime) const
+{
+	if(bDraggedOn && FSlateApplication::Get().GetDragDroppingContent().IsValid() && FSlateApplication::Get().GetDragDroppingContent()->IsOfType<FHierarchyDragDropOp>())
+	{
+		if(IsSectionSelected() == false)
+		{
+			OnSectionActivatedDelegate.ExecuteIfBound(SectionViewModel);
+		}
+	}
+
+	return EActiveTimerReturnType::Stop;
+}
+
+void SDataHierarchyEditor::Construct(const FArguments& InArgs, TObjectPtr<UDataHierarchyViewModelBase> InHierarchyViewModel)
+{
+	HierarchyViewModel = InHierarchyViewModel;
+
+	// If the user hasn't called initialize themselves, we do it here, but ideally the user should do it themselves where appropriate
+	if(HierarchyViewModel->IsInitialized() == false)
+	{
+		HierarchyViewModel->Initialize();
+	}
+	
+	SourceRoot.Reset(NewObject<UHierarchyRoot>(HierarchyViewModel->GetOuterForSourceRoot(), MakeUniqueObjectName(HierarchyViewModel->GetOuterForSourceRoot(), UHierarchyRoot::StaticClass()), RF_Transient));
+	TSharedPtr<FHierarchyElementViewModel> ViewModel = HierarchyViewModel->CreateViewModelForElement(SourceRoot.Get(), nullptr);
+	SourceRootViewModel = StaticCastSharedPtr<FHierarchyRootViewModel>(ViewModel);	
+	SourceRootViewModel->Initialize();
+	SourceRootViewModel->AddChildFilter(FHierarchyElementViewModel::FOnFilterChild::CreateSP(this, &SDataHierarchyEditor::FilterForSourceSection));
+	SourceRootViewModel->OnSyncPropagated().BindSP(this, &SDataHierarchyEditor::RequestRefreshSourceViewNextFrame, false);
+	SourceRootViewModel->OnSectionsChanged().BindSP(this, &SDataHierarchyEditor::RefreshSectionsView);
+
+	HierarchyViewModel->OnInitialized().BindSP(this, &SDataHierarchyEditor::Reinitialize);
+	HierarchyViewModel->OnNavigateToElementIdentityInHierarchyRequested().BindSP(this, &SDataHierarchyEditor::NavigateToHierarchyElement);
+	HierarchyViewModel->OnNavigateToElementInHierarchyRequested().BindSP(this, &SDataHierarchyEditor::NavigateToHierarchyElement);
+	HierarchyViewModel->OnRefreshSourceItemsRequested().BindSP(this, &SDataHierarchyEditor::RefreshSourceItems);
+	HierarchyViewModel->OnRefreshViewRequested().BindSP(this, &SDataHierarchyEditor::RefreshAllViews);
+	HierarchyViewModel->OnRefreshSourceView().BindSP(this, &SDataHierarchyEditor::RefreshSourceView);
+	HierarchyViewModel->OnRefreshHierarchyView().BindSP(this, &SDataHierarchyEditor::RefreshHierarchyView);
+	HierarchyViewModel->OnRefreshSectionsView().BindSP(this, &SDataHierarchyEditor::RefreshSectionsView);
+	HierarchyViewModel->OnHierarchySectionActivated().BindSP(this, &SDataHierarchyEditor::OnHierarchySectionActivated);
+	HierarchyViewModel->OnElementAdded().BindSP(this, &SDataHierarchyEditor::OnElementAdded);
+
+	BindToHierarchyRootViewModel();
+	
+	OnGenerateRowContentWidget = InArgs._OnGenerateRowContentWidget;
+	OnGenerateCustomDetailsPanelNameWidget = InArgs._OnGenerateCustomDetailsPanelNameWidget;
+	CategoryRowStyle = InArgs._CategoryRowStyle;
+	ItemRowStyle = InArgs._ItemRowStyle;
+	
+	if(!ensureMsgf(OnGenerateRowContentWidget.IsBound(), TEXT("Please add a function binding to the OnGenerateRowContentWidget slate event. Using default row content.")))
+	{
+		OnGenerateRowContentWidget = FOnGenerateRowContentWidget::CreateLambda([](TSharedRef<FHierarchyElementViewModel> HierarchyElement) -> TSharedRef<SWidget>
+		{
+			if(HierarchyElement->GetDataMutable()->IsA<UHierarchyCategory>())
+			{
+				TSharedRef<FHierarchyCategoryViewModel> TreeViewCategory = StaticCastSharedRef<FHierarchyCategoryViewModel>(HierarchyElement);
+				return SNew(SHierarchyCategory, TreeViewCategory);
+			}
+			
+			return SNew(STextBlock).Text(HierarchyElement->ToStringAsText());
+		});
+	}
+	SSplitter::FSlot* DetailsPanelSlot = nullptr;
+
+	TSharedRef<SWidget> AddSectionButton = SNew(SHorizontalBox)
+		+ SHorizontalBox::Slot()
+		.AutoWidth()
+		[
+			SNew(SButton)
+			.OnClicked(this, &SDataHierarchyEditor::OnAddSectionClicked)
+			.ButtonStyle(FDataHierarchyEditorStyle::Get(), "HierarchyEditor.ButtonStyle")
+			[
+				SNew(SHorizontalBox)
+				+ SHorizontalBox::Slot()
+				.AutoWidth()
+				.Padding(2.f)
+				[
+					SNew(SImage)
+					.Image(FAppStyle::Get().GetBrush("Icons.PlusCircle"))
+				]
+				+ SHorizontalBox::Slot()
+				.VAlign(VAlign_Center)
+				.AutoWidth()
+				[			
+					SNew(STextBlock)
+					.Text(LOCTEXT("AddSectionLabel","Add Section"))				
+				]
+			]
+		];
+
+	TSharedRef<SWidget> AddCategoryButton = SNew(SHorizontalBox)
+		+ SHorizontalBox::Slot()
+		.AutoWidth()
+		[
+			SNew(SButton)
+			.OnClicked(this, &SDataHierarchyEditor::OnAddCategoryClicked)
+			.ButtonStyle(FDataHierarchyEditorStyle::Get(), "HierarchyEditor.ButtonStyle")
+			[
+				SNew(SHorizontalBox)
+				+ SHorizontalBox::Slot()
+				.AutoWidth()
+				.Padding(2.f)
+				[
+					SNew(SImage)
+					.Image(FAppStyle::Get().GetBrush("Icons.PlusCircle"))
+				]
+				+ SHorizontalBox::Slot()
+				.VAlign(VAlign_Center)
+				.AutoWidth()
+				[			
+					SNew(STextBlock)
+					.Text(LOCTEXT("AddCategoryLabel","Add Category"))				
+				]
+			]
+		];
+	
+	ChildSlot
+	[
+		SNew(SBorder)
+		.Padding(0.f)
+		.BorderImage(FAppStyle::Get().GetBrush("Brushes.Panel"))
+		[
+			SNew(SSplitter)
+			.Orientation(Orient_Horizontal)
+			.PhysicalSplitterHandleSize(2.f)
+			+ SSplitter::Slot()
+			.Value(0.3f)
+			.MinSize(0.1f)
+			[
+				SNew(SVerticalBox)
+				+ SVerticalBox::Slot()
+				.AutoHeight()
+				.Padding(2.f)
+				[
+					SAssignNew(SourceSearchBox, SSearchBox)
+					.OnTextChanged(this, &SDataHierarchyEditor::OnSourceSearchTextChanged)
+					.OnTextCommitted(this, &SDataHierarchyEditor::OnSourceSearchTextCommitted)
+					.OnSearch(SSearchBox::FOnSearch::CreateSP(this, &SDataHierarchyEditor::OnSearchButtonClicked))
+					.DelayChangeNotificationsWhileTyping(true)
+					.SearchResultData(this, &SDataHierarchyEditor::GetSearchResultData)
+				]
+				+ SVerticalBox::Slot()
+				.AutoHeight()
+				.Padding(2.f)
+				[
+					SAssignNew(SourceSectionBox, SWrapBox)
+					.UseAllottedSize(true)
+				]
+				+ SVerticalBox::Slot()
+				.Padding(1.f, 2.f)
+				[
+					SAssignNew(SourceTreeView, STreeView<TSharedPtr<FHierarchyElementViewModel>>)
+					.TreeItemsSource(&GetSourceItems())
+					.OnSelectionChanged(this, &SDataHierarchyEditor::OnSelectionChanged, false)
+					.OnGenerateRow(this, &SDataHierarchyEditor::GenerateSourceItemRow)
+					.OnGetChildren_UObject(InHierarchyViewModel.Get(), &UDataHierarchyViewModelBase::OnGetChildren)
+					.OnItemToString_Debug_UObject(InHierarchyViewModel.Get(), &UDataHierarchyViewModelBase::OnElementToStringDebug)
+					.OnContextMenuOpening(this, &SDataHierarchyEditor::SummonContextMenuForSelectedRows, false)
+				]					
+			]
+			+ SSplitter::Slot()
+			.Value(0.4f)
+			.MinSize(0.1f)
+			[
+				SNew(SVerticalBox)
+				+ SVerticalBox::Slot()
+				.AutoHeight()
+				[
+					SNew(SVerticalBox)
+					+ SVerticalBox::Slot()
+					.AutoHeight()
+					[
+						SNew(SBorder)
+						.Padding(0.f)
+						.BorderImage(FAppStyle::Get().GetBrush("Brushes.Header"))
+						[
+							AddSectionButton
+						]
+					]
+					+ SVerticalBox::Slot()
+					.AutoHeight()
+					.Padding(1.f)
+					[
+						SAssignNew(HierarchySectionBox, SWrapBox)
+						.UseAllottedSize(true)
+					]				
+				]
+				+ SVerticalBox::Slot()
+				[
+					SNew(SVerticalBox)
+					+ SVerticalBox::Slot()
+					.AutoHeight()
+					[
+						SNew(SBorder)
+						.Padding(0.f)
+						.BorderImage(FAppStyle::Get().GetBrush("Brushes.Header"))
+						[
+							AddCategoryButton
+						]
+					]
+					+ SVerticalBox::Slot()
+					.FillHeight(0.1f)
+					.Padding(1.f, 4.f, 1.f, 0.f)
+					[
+						SNew(SDropTarget)
+						.OnDropped(this, &SDataHierarchyEditor::HandleHierarchyRootDrop)
+						.OnAllowDrop(this, &SDataHierarchyEditor::OnCanDropOnRoot)
+						.OnDragEnter(this, &SDataHierarchyEditor::OnRootDragEnter)
+						.OnDragLeave(this, &SDataHierarchyEditor::OnRootDragLeave)
+						[
+							SNew(SBorder)
+							.Padding(0.f)
+							.BorderImage(FAppStyle::Get().GetBrush("Brushes.Recessed"))
+							[
+								SNew(SBorder)
+								.Padding(1.f)
+								.BorderImage(FAppStyle::GetBrush("DashedBorder"))
+								.BorderBackgroundColor(FLinearColor(0.2f, 0.2f, 0.2f, 0.5f))
+								[
+									SNew(SBox)
+									.HAlign(HAlign_Center)
+									.VAlign(VAlign_Center)
+									[
+										SNew(SImage)
+										.Image(FDataHierarchyEditorStyle::Get().GetBrush("HierarchyEditor.RootDropIcon"))
+										.ColorAndOpacity(this, &SDataHierarchyEditor::GetRootIconColor)
+									]
+								]
+							]
+						]
+					]
+					+ SVerticalBox::Slot()
+					.Padding(1.f, 0.f)
+					[
+						SAssignNew(HierarchyTreeView, STreeView<TSharedPtr<FHierarchyElementViewModel>>)
+						.TreeItemsSource(&InHierarchyViewModel->GetHierarchyItems())
+						.OnSelectionChanged(this, &SDataHierarchyEditor::OnSelectionChanged, true)
+						.OnGenerateRow(this, &SDataHierarchyEditor::GenerateHierarchyItemRow)
+						.OnGetChildren_UObject(InHierarchyViewModel.Get(), &UDataHierarchyViewModelBase::OnGetChildren)
+						.OnItemToString_Debug_UObject(InHierarchyViewModel.Get(), &UDataHierarchyViewModelBase::OnElementToStringDebug)
+						.OnContextMenuOpening(this, &SDataHierarchyEditor::SummonContextMenuForSelectedRows, true)
+					]
+				]
+			]
+			+ SSplitter::Slot()
+			.Value(0.3f)
+			.MinSize(0.1f)
+			.Expose(DetailsPanelSlot)		
+		]
+	];
+
+	if(InHierarchyViewModel->SupportsDetailsPanel())
+	{
+		FPropertyEditorModule& PropertyEditorModule = FModuleManager::Get().LoadModuleChecked<FPropertyEditorModule>("PropertyEditor");
+
+		FDetailsViewArgs DetailsViewArgs;
+		DetailsViewArgs.NameAreaSettings = FDetailsViewArgs::ENameAreaSettings::ObjectsUseNameArea;
+		DetailsViewArgs.bShowObjectLabel = false;
+		DetailsViewArgs.bAllowSearch = false;
+		DetailsViewArgs.NotifyHook = this;
+
+		DetailsPanel = PropertyEditorModule.CreateDetailView(DetailsViewArgs);
+		
+		if(OnGenerateCustomDetailsPanelNameWidget.IsBound())
+		{
+			TSharedRef<SWidget> CustomDetailsPanelNameWidget = OnGenerateCustomDetailsPanelNameWidget.Execute(nullptr);
+			DetailsPanel->SetNameAreaCustomContent(CustomDetailsPanelNameWidget);
+		}
+
+		DetailsPanel->SetIsPropertyEditingEnabledDelegate(FIsPropertyEditingEnabled::CreateSP(this, &SDataHierarchyEditor::IsDetailsPanelEditingAllowed));
+		
+		for(auto& Customizations : HierarchyViewModel->GetInstanceCustomizations())
+		{
+			DetailsPanel->RegisterInstancedCustomPropertyLayout(Customizations.Key, Customizations.Value);
+		}
+		
+		DetailsPanelSlot->AttachWidget(DetailsPanel.ToSharedRef());
+	}
+		
+	HierarchyViewModel->GetCommands()->MapAction(FGenericCommands::Get().Rename,
+		FExecuteAction::CreateSP(this, &SDataHierarchyEditor::RequestRenameSelectedItem),
+		FCanExecuteAction::CreateSP(this, &SDataHierarchyEditor::CanRequestRenameSelectedItem),
+		FIsActionChecked(), FIsActionButtonVisible::CreateSP(this, &SDataHierarchyEditor::CanRequestRenameSelectedItem));
+	
+	HierarchyViewModel->GetCommands()->MapAction(FGenericCommands::Get().Delete,
+		FExecuteAction::CreateSP(this, &SDataHierarchyEditor::DeleteSelectedHierarchyItems),
+		FCanExecuteAction::CreateSP(this, &SDataHierarchyEditor::CanDeleteSelectedElements),
+		FIsActionChecked(), FIsActionButtonVisible::CreateSP(this, &SDataHierarchyEditor::CanDeleteSelectedElements));
+	
+	HierarchyViewModel->GetCommands()->MapAction(FDataHierarchyEditorCommands::Get().FindInHierarchy,
+		FExecuteAction::CreateSP(this, &SDataHierarchyEditor::NavigateToMatchingHierarchyElementFromSelectedSourceElement),
+		FCanExecuteAction::CreateSP(this, &SDataHierarchyEditor::CanNavigateToMatchingHierarchyElementFromSelectedSourceElement),
+		FIsActionChecked(), FIsActionButtonVisible::CreateSP(this, &SDataHierarchyEditor::CanNavigateToMatchingHierarchyElementFromSelectedSourceElement));
+	
+	HierarchyViewModel->ForceFullRefresh();
+}
+
+SDataHierarchyEditor::~SDataHierarchyEditor()
+{
+	SourceSearchResults.Empty();
+	FocusedSearchResult.Reset();
+	
+	ClearSourceItems();
+	
+	if(HierarchyViewModel.IsValid())
+	{
+		HierarchyViewModel->OnInitialized().Unbind();
+		HierarchyViewModel->OnNavigateToElementIdentityInHierarchyRequested().Unbind();
+		HierarchyViewModel->OnNavigateToElementInHierarchyRequested().Unbind();
+		HierarchyViewModel->OnRefreshSourceItemsRequested().Unbind();
+		HierarchyViewModel->OnRefreshViewRequested().Unbind();
+		HierarchyViewModel->OnRefreshSourceView().Unbind();
+		HierarchyViewModel->OnRefreshHierarchyView().Unbind();
+		HierarchyViewModel->OnRefreshSectionsView().Unbind();
+		HierarchyViewModel->OnHierarchySectionActivated().Unbind();
+		HierarchyViewModel->OnElementAdded().Unbind();
+
+		UnbindFromHierarchyRootViewModel();
+
+		HierarchyViewModel->GetCommands()->UnmapAction(FGenericCommands::Get().Delete);
+		HierarchyViewModel->GetCommands()->UnmapAction(FGenericCommands::Get().Rename);
+		HierarchyViewModel->GetCommands()->UnmapAction(FDataHierarchyEditorCommands::Get().FindInHierarchy);
+	}
+	
+	SourceRootViewModel->OnSyncPropagated().Unbind();
+	SourceRootViewModel->OnSectionsChanged().Unbind();
+	SourceRootViewModel.Reset();
+	SourceRoot->ConditionalBeginDestroy();
+	SourceRoot = nullptr;
+}
+
+void SDataHierarchyEditor::RefreshSourceItems()
+{
+	SourceRoot->EmptyAllData();
+	HierarchyViewModel->PrepareSourceItems(SourceRoot.Get(), SourceRootViewModel);
+	SourceRootViewModel->SyncViewModelsToData();
+	RefreshSourceView(false);
+	RefreshSectionsView();
+}
+
+void SDataHierarchyEditor::RefreshAllViews(bool bFullRefresh)
+{
+	RefreshSourceView(bFullRefresh);
+	RefreshHierarchyView(bFullRefresh);
+	RefreshSectionsView();
+}
+
+void SDataHierarchyEditor::RequestRefreshAllViewsNextFrame(bool bFullRefresh)
+{
+	RequestRefreshSourceViewNextFrame(bFullRefresh);
+	RequestRefreshHierarchyViewNextFrame(bFullRefresh);
+	RequestRefreshSectionsViewNextFrame();
+}
+
+FReply SDataHierarchyEditor::OnKeyDown(const FGeometry& MyGeometry, const FKeyEvent& InKeyEvent)
+{
+	UHierarchyMenuContext* HierarchyMenuContext = NewObject<UHierarchyMenuContext>();
+	HierarchyMenuContext->HierarchyViewModel = HierarchyViewModel;
+	
+	FToolMenuContext Context;
+	Context.AddObject(HierarchyMenuContext);
+	
+	UToolMenus::Get()->GenerateMenu(HierarchyViewModel->GetContextMenuName(), Context);
+	return HierarchyViewModel->GetCommands()->ProcessCommandBindings(InKeyEvent) ? FReply::Handled() : FReply::Unhandled();
+}
+
+FReply SDataHierarchyEditor::OnMouseButtonDown(const FGeometry& MyGeometry, const FPointerEvent& MouseEvent)
+{
+	// we catch any mouse button down event so that we can continue using our commands
+	return FReply::Handled().SetUserFocus(AsShared(), EFocusCause::Mouse, true);
+}
+
+FReply SDataHierarchyEditor::OnMouseButtonUp(const FGeometry& MyGeometry, const FPointerEvent& MouseEvent)
+{
+	return FReply::Handled().SetUserFocus(AsShared(), EFocusCause::Mouse, true);
+}
+
+FReply SDataHierarchyEditor::OnAddCategoryClicked() const
+{
+	TArray<TSharedPtr<FHierarchyElementViewModel>> SelectedItems;
+	HierarchyTreeView->GetSelectedItems(SelectedItems);
+
+	// We can only add categories under categories or the root
+	if(SelectedItems.Num() == 1 && SelectedItems[0]->GetData()->IsA<UHierarchyCategory>())
+	{
+		HierarchyViewModel.Get()->AddCategory(SelectedItems[0]);
+	}
+	else
+	{
+		HierarchyViewModel.Get()->AddCategory(nullptr);
+	}
+	
+	return FReply::Handled();
+}
+
+FReply SDataHierarchyEditor::OnAddSectionClicked() const
+{
+	HierarchyViewModel.Get()->AddSection();
+	return FReply::Handled();
+}
+
+TSharedPtr<SWidget> SDataHierarchyEditor::SummonContextMenuForSelectedRows(bool bFromHierarchy) const
+{
+	TArray<TSharedPtr<FHierarchyElementViewModel>> ViewModels;
+	if(bFromHierarchy)
+	{
+		HierarchyTreeView->GetSelectedItems(ViewModels);
+	}
+	else
+	{
+		SourceTreeView->GetSelectedItems(ViewModels);
+	}
+
+	if(ViewModels.IsEmpty())
+	{
+		return nullptr;
+	}
+	
+	return SummonContextMenu(ViewModels);
+}
+
+void SDataHierarchyEditor::RefreshSourceView(bool bFullRefresh) const
+{
+	SourceTreeView->SetTreeItemsSource(&GetSourceItems());
+	if(bFullRefresh)
+	{
+		SourceTreeView->RebuildList();
+	}
+	else
+	{
+		SourceTreeView->RequestTreeRefresh();
+	}
+}
+
+void SDataHierarchyEditor::RequestRefreshSourceViewNextFrame(bool bFullRefresh)
+{
+	if(!RefreshSourceViewNextFrameHandle.IsValid())
+	{
+		RefreshSourceViewNextFrameHandle = RegisterActiveTimer(0.f, FWidgetActiveTimerDelegate::CreateLambda([this, bFullRefresh](double CurrentTime, float DeltaTime)
+		{
+			RefreshSourceView(bFullRefresh);
+			RefreshSourceViewNextFrameHandle.Reset();
+			return EActiveTimerReturnType::Stop;
+		}));
+	}
+}
+
+void SDataHierarchyEditor::RefreshHierarchyView(bool bFullRefresh) const
+{
+	// the top layer objects might have changed due to filtering. We need to refresh these too.
+	HierarchyTreeView->SetTreeItemsSource(&HierarchyViewModel->GetHierarchyItems());
+	if(bFullRefresh)
+	{
+		HierarchyTreeView->RebuildList();
+	}
+	else
+	{
+		HierarchyTreeView->RequestTreeRefresh();
+	}
+}
+
+void SDataHierarchyEditor::RequestRefreshHierarchyViewNextFrame(bool bFullRefresh)
+{
+	if(!RefreshHierarchyViewNextFrameHandle.IsValid())
+	{
+		RefreshHierarchyViewNextFrameHandle = RegisterActiveTimer(0.f, FWidgetActiveTimerDelegate::CreateLambda([this, bFullRefresh](double CurrentTime, float DeltaTime)
+		{
+			RefreshHierarchyView(bFullRefresh);
+			RefreshHierarchyViewNextFrameHandle.Reset();
+			return EActiveTimerReturnType::Stop;
+		}));
+	}
+}
+
+void SDataHierarchyEditor::RefreshSectionsView()
+{
+	SourceSectionBox->ClearChildren();
+	HierarchySectionBox->ClearChildren();
+
+	for (TSharedPtr<FHierarchySectionViewModel>& SourceSection : SourceRootViewModel->GetSectionViewModels())
+	{
+		TSharedPtr<SHierarchySection> SectionWidget = SNew(SHierarchySection, SourceSection, HierarchyViewModel)
+		.IsSectionActive_Lambda([this, SourceSection]()
+		{
+			return GetActiveSourceSection() == SourceSection ? ECheckBoxState::Checked : ECheckBoxState::Unchecked;
+		})
+		.OnSectionActivated_Lambda([this](TSharedPtr<FHierarchySectionViewModel> SectionViewModel)
+		{
+			SetActiveSourceSection(SectionViewModel);
+		});
+		
+		SourceSectionBox->AddSlot()
+		.HAlign(HAlign_Center)
+		.VAlign(VAlign_Center)
+		[
+			SectionWidget.ToSharedRef()		
+		];
+	}
+
+	if(SourceRootViewModel->GetSectionViewModels().Num() > 0)
+	{
+		TSharedPtr<SHierarchySection> DefaultSourceSection = SNew(SHierarchySection, nullptr, HierarchyViewModel)
+		.IsSectionActive_Lambda([this]()
+		{
+			return GetActiveSourceSection() == nullptr ? ECheckBoxState::Checked : ECheckBoxState::Unchecked;
+		})
+		.OnSectionActivated(this, &SDataHierarchyEditor::SetActiveSourceSection)
+		// we forbid drop on here as the 'All' sections don't have a valid view model to determine it instead
+		.bForbidDropOn(true);
+
+		SourceSectionBox->AddSlot()
+		.HAlign(HAlign_Center)
+		.VAlign(VAlign_Center)
+		[
+			DefaultSourceSection.ToSharedRef()
+		];
+	}
+	
+	for (TSharedPtr<FHierarchySectionViewModel>& HierarchySection : HierarchyViewModel->GetHierarchyRootViewModel()->GetSectionViewModels())
+	{
+		TSharedPtr<SHierarchySection> SectionWidget = SNew(SHierarchySection, HierarchySection, HierarchyViewModel)
+		.IsSectionActive_Lambda([this, HierarchySection]()
+		{
+			return HierarchyViewModel->GetActiveHierarchySection() == HierarchySection ? ECheckBoxState::Checked : ECheckBoxState::Unchecked;
+		})
+		.OnSectionActivated_Lambda([this](TSharedPtr<FHierarchySectionViewModel> SectionViewModel)
+		{
+			HierarchyViewModel->SetActiveHierarchySection(SectionViewModel);
+		});
+		
+		HierarchySectionBox->AddSlot()
+		.HAlign(HAlign_Center)
+		.VAlign(VAlign_Center)
+		[
+			SectionWidget.ToSharedRef()		
+		];
+	}
+	
+	TSharedPtr<SHierarchySection> DefaultHierarchySection = SNew(SHierarchySection, nullptr, HierarchyViewModel)
+	.IsSectionActive_Lambda([this]()
+	{
+		return HierarchyViewModel->GetActiveHierarchySection() == nullptr ? ECheckBoxState::Checked : ECheckBoxState::Unchecked;
+	})
+	.OnSectionActivated_Lambda([this](TSharedPtr<FHierarchySectionViewModel> SectionViewModel)
+	{
+		HierarchyViewModel->SetActiveHierarchySection(SectionViewModel);
+	});
+
+	HierarchySectionBox->AddSlot()
+	.HAlign(HAlign_Center)
+	.VAlign(VAlign_Center)
+	[
+		DefaultHierarchySection.ToSharedRef()
+	];
+}
+
+void SDataHierarchyEditor::RequestRefreshSectionsViewNextFrame()
+{
+	if(!RefreshSectionsViewNextFrameHandle.IsValid())
+	{
+		RefreshSectionsViewNextFrameHandle = RegisterActiveTimer(0.f, FWidgetActiveTimerDelegate::CreateLambda([this](double CurrentTime, float DeltaTime)
+		{
+			RefreshSectionsView();
+			RefreshSectionsViewNextFrameHandle.Reset();
+			return EActiveTimerReturnType::Stop;
+		}));
+	}
+}
+
+void SDataHierarchyEditor::NavigateToHierarchyElement(FHierarchyElementIdentity Identity) const
+{
+	if(TSharedPtr<FHierarchyElementViewModel> ViewModel = HierarchyViewModel->GetHierarchyRootViewModel()->FindViewModelForChild(Identity, true))
+	{
+		NavigateToHierarchyElement(ViewModel);
+	}
+}
+
+void SDataHierarchyEditor::NavigateToHierarchyElement(TSharedPtr<FHierarchyElementViewModel> Item) const
+{
+	TArray<TSharedPtr<FHierarchyElementViewModel>> ParentChain;
+	for(TWeakPtr<FHierarchyElementViewModel> Parent = Item->GetParent(); Parent.IsValid(); Parent = Parent.Pin()->GetParent())
+	{
+		ParentChain.Add(Parent.Pin());
+	}
+
+	for(int32 ParentIndex = ParentChain.Num()-1; ParentIndex >= 0; ParentIndex--)
+	{
+		HierarchyTreeView->SetItemExpansion(ParentChain[ParentIndex], true);
+	}
+		
+	HierarchyTreeView->SetSelection(Item);
+	HierarchyTreeView->RequestScrollIntoView(Item);
+}
+
+bool SDataHierarchyEditor::IsItemSelected(TSharedPtr<FHierarchyElementViewModel> Item) const
+{
+	return HierarchyTreeView->IsItemSelected(Item);
+}
+
+TSharedRef<ITableRow> SDataHierarchyEditor::GenerateSourceItemRow(TSharedPtr<FHierarchyElementViewModel> HierarchyItem, const TSharedRef<STableViewBase>& TableViewBase)
+{
+	return SNew(STableRow<TSharedPtr<FHierarchyElementViewModel>>, TableViewBase)
+	.Style(HierarchyItem->GetData()->IsA<UHierarchyItem>() ? ItemRowStyle : CategoryRowStyle)
+	.OnDragDetected(HierarchyItem.ToSharedRef(), &FHierarchyElementViewModel::OnDragDetected, true)
+	.Padding(FMargin(2))
+	[
+		SNew(SHorizontalBox)
+		+ SHorizontalBox::Slot()
+		.AutoWidth()
+		.Padding(1.f)
+		.HAlign(HAlign_Center)
+		.VAlign(VAlign_Center)
+		[
+			SNew(SBox)
+			.HeightOverride(10.f)
+			.WidthOverride(10.f)
+			.Visibility_Lambda([HierarchyItem, this]()
+			{
+				TArray<TSharedPtr<FHierarchyElementViewModel>> AllChildren;
+				HierarchyItem->GetChildrenViewModelsForType<UHierarchyElement, FHierarchyElementViewModel>(AllChildren, true);
+
+				bool bCanDrag = HierarchyItem->GetHierarchyViewModel()->GetHierarchyRootViewModel()->FindViewModelForChild(HierarchyItem->GetData()->GetPersistentIdentity(), true) == nullptr;			
+
+				if(bCanDrag)
+				{
+					for(TSharedPtr<FHierarchyElementViewModel>& ItemViewModel : AllChildren)
+					{
+						if(HierarchyItem->GetHierarchyViewModel()->GetHierarchyRootViewModel()->FindViewModelForChild(ItemViewModel->GetData()->GetPersistentIdentity(), true) != nullptr)
+						{
+							bCanDrag = false;
+							break;
+						}
+					}
+				}
+
+				return bCanDrag ? EVisibility::Collapsed : EVisibility::Visible;
+			})
+			[
+				SNew(SImage)
+				.Image(FAppStyle::GetBrush("Icons.Lock"))
+				.ToolTipText(LOCTEXT("CantDragItemAlreadyInHierarchyTooltip", "This item already exists within the hierarchy and can not be dragged. Drag the existing one within the hierarchy directly."))
+			]
+		]
+		+ SHorizontalBox::Slot()
+		[
+			OnGenerateRowContentWidget.Execute(HierarchyItem.ToSharedRef())
+		]
+	];
+}
+
+TSharedRef<ITableRow> SDataHierarchyEditor::GenerateHierarchyItemRow(TSharedPtr<FHierarchyElementViewModel> HierarchyItem, const TSharedRef<STableViewBase>& TableViewBase)
+{
+	return SNew(STableRow<TSharedPtr<FHierarchyElementViewModel>>, TableViewBase)
+	.Style(HierarchyItem->GetData()->IsA<UHierarchyItem>() ? ItemRowStyle : CategoryRowStyle)
+	.OnAcceptDrop(HierarchyItem.ToSharedRef(), &FHierarchyElementViewModel::OnDroppedOnRow)
+	.OnCanAcceptDrop(HierarchyItem.ToSharedRef(), &FHierarchyElementViewModel::OnCanRowAcceptDrop)
+	.OnDragDetected(HierarchyItem.ToSharedRef(), &FHierarchyElementViewModel::OnDragDetected, false)
+	.OnDragLeave(HierarchyItem.ToSharedRef(), &FHierarchyElementViewModel::OnRowDragLeave)
+	.Padding(FMargin(2))
+	[
+		SNew(SHorizontalBox)
+		+ SHorizontalBox::Slot()
+		.AutoWidth()
+		.Padding(1.f)
+		.HAlign(HAlign_Center)
+		.VAlign(VAlign_Center)
+		[
+			SNew(SBox)
+			.HeightOverride(10.f)
+			.WidthOverride(10.f)
+			.Visibility_Lambda([HierarchyItem]()
+			{
+				return HierarchyItem->IsEditableByUser().bCanPerform ? EVisibility::Collapsed : EVisibility::Visible;
+			})
+			[
+				SNew(SImage)
+				.Image(FAppStyle::GetBrush("Icons.Lock"))
+				.ToolTipText_Lambda([HierarchyItem]()
+				{
+					FHierarchyElementViewModel::FCanPerformActionResults IsEditableResults = HierarchyItem->IsEditableByUser();
+					if(IsEditableResults.bCanPerform == false)
+					{
+						return IsEditableResults.CanPerformMessage;
+					}
+
+					return FText::GetEmpty();
+				})
+			]
+		]
+		+ SHorizontalBox::Slot()
+		[
+			OnGenerateRowContentWidget.Execute(HierarchyItem.ToSharedRef())
+		]
+	];
+}
+
+bool SDataHierarchyEditor::FilterForSourceSection(TSharedPtr<const FHierarchyElementViewModel> ItemViewModel) const
+{
+	if(ActiveSourceSection.IsValid())
+	{
+		return GetActiveSourceSectionData() == ItemViewModel->GetSection();
+	}
+
+	return true;
+}
+
+void SDataHierarchyEditor::Reinitialize()
+{
+	// the hierarchy root view model has been recreated if the view model reinitialized. Therefore we update the bindings.
+	BindToHierarchyRootViewModel();
+	RefreshSourceItems();
+	RefreshAllViews(true);
+}
+
+void SDataHierarchyEditor::BindToHierarchyRootViewModel()
+{
+	HierarchyViewModel->GetHierarchyRootViewModel()->OnSyncPropagated().BindSP(this, &SDataHierarchyEditor::RequestRefreshHierarchyViewNextFrame, false);
+	HierarchyViewModel->GetHierarchyRootViewModel()->OnSectionsChanged().BindSP(this, &SDataHierarchyEditor::RefreshSectionsView);
+	HierarchyViewModel->GetHierarchyRootViewModel()->OnSectionAdded().BindSP(this, &SDataHierarchyEditor::OnHierarchySectionAdded);
+	HierarchyViewModel->GetHierarchyRootViewModel()->OnSectionDeleted().BindSP(this, &SDataHierarchyEditor::OnHierarchySectionDeleted);
+}
+
+void SDataHierarchyEditor::UnbindFromHierarchyRootViewModel() const
+{
+	if(HierarchyViewModel->GetHierarchyRootViewModel().IsValid())
+	{
+		HierarchyViewModel->GetHierarchyRootViewModel()->OnSyncPropagated().Unbind();
+		HierarchyViewModel->GetHierarchyRootViewModel()->OnSectionsChanged().Unbind();
+		HierarchyViewModel->GetHierarchyRootViewModel()->OnSectionAdded().Unbind();
+		HierarchyViewModel->GetHierarchyRootViewModel()->OnSectionDeleted().Unbind();
+	}
+}
+
+const TArray<TSharedPtr<FHierarchyElementViewModel>>& SDataHierarchyEditor::GetSourceItems() const
+{
+	return SourceRootViewModel->GetFilteredChildren();
+}
+
+bool SDataHierarchyEditor::IsDetailsPanelEditingAllowed() const
+{
+	return SelectedDetailsPanelItemViewModel.IsValid() && SelectedDetailsPanelItemViewModel.Pin()->IsEditableByUser().bCanPerform;
+}
+
+void SDataHierarchyEditor::RequestRenameSelectedItem()
+{
+	TArray<TSharedPtr<FHierarchyElementViewModel>> SelectedItems = HierarchyTreeView->GetSelectedItems();
+
+	if(SelectedItems.Num() == 0)
+	{
+		TSharedPtr<FHierarchySectionViewModel> ActiveSection = HierarchyViewModel->GetActiveHierarchySection();
+		if(ActiveSection)
+		{
+			SelectedItems = { ActiveSection };
+		}
+	}
+	
+	if(SelectedItems.Num() == 1)
+	{
+		SelectedItems[0]->RequestRename();
+	}
+}
+
+bool SDataHierarchyEditor::CanRequestRenameSelectedItem() const
+{
+	TArray<TSharedPtr<FHierarchyElementViewModel>> SelectedItems = HierarchyTreeView->GetSelectedItems();
+
+	if(SelectedItems.Num() == 0)
+	{
+		TSharedPtr<FHierarchySectionViewModel> ActiveSection = HierarchyViewModel->GetActiveHierarchySection();
+		if(ActiveSection)
+		{
+			SelectedItems = { ActiveSection };
+		}
+	}
+	
+	if(SelectedItems.Num() == 1)
+	{
+		return SelectedItems[0]->CanRename();
+	}
+
+	return false;
+}
+
+void SDataHierarchyEditor::ClearSourceItems() const
+{
+	SourceRoot->GetChildrenMutable().Empty();
+	SourceRoot->GetSectionDataMutable().Empty();
+	SourceRootViewModel->GetChildrenMutable().Empty();
+	SourceRootViewModel->GetSectionViewModels().Empty();
+}
+
+void SDataHierarchyEditor::DeleteItems(TArray<TSharedPtr<FHierarchyElementViewModel>> ItemsToDelete) const
+{
+	TArray<FHierarchyElementIdentity> DeletionIdentities;
+	
+	for(TSharedPtr<FHierarchyElementViewModel>& SelectedItem : ItemsToDelete)
+	{
+		DeletionIdentities.Add(SelectedItem->GetData()->GetPersistentIdentity());
+	}
+
+	HierarchyViewModel->DeleteElementsWithIdentities(DeletionIdentities);
+}
+
+void SDataHierarchyEditor::DeleteSelectedHierarchyItems() const
+{	
+	TArray<TSharedPtr<FHierarchyElementViewModel>> SelectedElements = HierarchyTreeView->GetSelectedItems();
+	
+	if(SelectedElements.Num() == 0)
+	{
+		TSharedPtr<FHierarchySectionViewModel> ActiveSection = HierarchyViewModel->GetActiveHierarchySection();
+		if(ActiveSection)
+		{
+			SelectedElements = { ActiveSection };
+		}
+	}
+	
+	DeleteItems(SelectedElements);
+}
+
+bool SDataHierarchyEditor::CanDeleteSelectedElements() const
+{
+	TArray<TSharedPtr<FHierarchyElementViewModel>> SelectedElements = HierarchyTreeView->GetSelectedItems();
+
+	if(SelectedElements.Num() == 0)
+	{
+		TSharedPtr<FHierarchySectionViewModel> ActiveSection = HierarchyViewModel->GetActiveHierarchySection();
+		if(ActiveSection)
+		{
+			SelectedElements = { ActiveSection };
+		}
+	}
+	
+	if(SelectedElements.Num() > 0)
+	{
+		bool bCanDelete = true;
+		for(TSharedPtr<FHierarchyElementViewModel> SelectedElement : SelectedElements)
+		{
+			bCanDelete &= SelectedElements[0]->CanDelete();
+		}
+
+		return bCanDelete;
+	}
+
+	return false;
+}
+
+void SDataHierarchyEditor::NavigateToMatchingHierarchyElementFromSelectedSourceElement() const
+{
+	TArray<TSharedPtr<FHierarchyElementViewModel>> SelectedElements = SourceTreeView->GetSelectedItems();
+
+	if(SelectedElements.Num() != 1)
+	{
+		return;
+	}
+
+	TSharedPtr<FHierarchyElementViewModel> SelectedElement = SelectedElements[0];
+
+	if(TSharedPtr<FHierarchyElementViewModel> MatchingViewModelInHierarchy = HierarchyViewModel->GetHierarchyRootViewModel()->FindViewModelForChild(SelectedElement->GetData()->GetPersistentIdentity(), true))
+	{
+		NavigateToHierarchyElement(MatchingViewModelInHierarchy.ToSharedRef());
+	}
+}
+
+bool SDataHierarchyEditor::CanNavigateToMatchingHierarchyElementFromSelectedSourceElement() const
+{
+	TArray<TSharedPtr<FHierarchyElementViewModel>> SelectedElements = SourceTreeView->GetSelectedItems();
+
+	if(SelectedElements.Num() != 1)
+	{
+		return false;
+	}
+
+	TSharedPtr<FHierarchyElementViewModel> SelectedElement = SelectedElements[0];
+	if(SelectedElement->IsForHierarchy())
+	{
+		return false;
+	}
+	
+	if(TSharedPtr<FHierarchyElementViewModel> MatchingViewModelInHierarchy = HierarchyViewModel->GetHierarchyRootViewModel()->FindViewModelForChild(SelectedElement->GetData()->GetPersistentIdentity(), true))
+	{
+		return MatchingViewModelInHierarchy.IsValid();
+	}
+	
+	return false;
+}
+
+void SDataHierarchyEditor::DeleteActiveSection() const
+{
+	if(TSharedPtr<FHierarchySectionViewModel> ActiveHierarchySectionViewModel = HierarchyViewModel->GetActiveHierarchySection())
+	{
+		DeleteItems({ActiveHierarchySectionViewModel});
+	}
+}
+
+bool SDataHierarchyEditor::CanDeleteActiveSection() const
+{
+	return HierarchyViewModel->GetActiveHierarchySection() != nullptr;
+}
+
+void SDataHierarchyEditor::OnElementAdded(TSharedPtr<FHierarchyElementViewModel> AddedItem)
+{
+	// when a new item is created (opposed to dragged & dropped from source view, i.e. only categories so far)
+	// we make sure to request a tree refresh, select the row, and request a pending rename since the widget will created a frame later
+	if(AddedItem->GetData()->IsA<UHierarchyItem>() || AddedItem->GetData()->IsA<UHierarchyCategory>())
+	{
+		HierarchyTreeView->RequestTreeRefresh();
+		NavigateToHierarchyElement(AddedItem);
+	}
+	else if(AddedItem->GetData()->IsA<UHierarchySection>())
+	{
+		RefreshSectionsView();
+	}
+
+	AddedItem->RequestRenamePending();
+}
+
+void SDataHierarchyEditor::OnHierarchySectionActivated(TSharedPtr<FHierarchySectionViewModel> Section)
+{
+	OnSelectionChanged(Section, ESelectInfo::Direct, true);
+}
+
+void SDataHierarchyEditor::OnSourceSectionActivated(TSharedPtr<FHierarchySectionViewModel> Section)
+{
+	OnSelectionChanged(Section, ESelectInfo::Direct, false);
+	RunSourceSearch();
+}
+
+void SDataHierarchyEditor::OnHierarchySectionAdded(TSharedPtr<FHierarchySectionViewModel> AddedSection)
+{
+	HierarchyViewModel->SetActiveHierarchySection(AddedSection);
+	AddedSection->RequestRenamePending();
+}
+
+void SDataHierarchyEditor::OnHierarchySectionDeleted(TSharedPtr<FHierarchySectionViewModel> DeletedSection)
+{
+	if(HierarchyViewModel->GetActiveHierarchySection() == DeletedSection)
+	{
+		HierarchyViewModel->SetActiveHierarchySection(nullptr);
+	}
+}
+
+void SDataHierarchyEditor::SetActiveSourceSection(TSharedPtr<FHierarchySectionViewModel> Section)
+{
+	ActiveSourceSection = Section;	
+	RefreshSourceView(true);
+	OnSourceSectionActivated(Section);
+}
+
+TSharedPtr<FHierarchySectionViewModel> SDataHierarchyEditor::GetActiveSourceSection() const
+{
+	return ActiveSourceSection.IsValid() ? ActiveSourceSection.Pin() : nullptr;
+}
+
+UHierarchySection* SDataHierarchyEditor::GetActiveSourceSectionData() const
+{
+	return ActiveSourceSection.IsValid() ? ActiveSourceSection.Pin()->GetDataMutable<UHierarchySection>() : nullptr;
+}
+
+void SDataHierarchyEditor::OnSelectionChanged(TSharedPtr<FHierarchyElementViewModel> HierarchyItem, ESelectInfo::Type Type, bool bFromHierarchy) const
+{
+	SelectedDetailsPanelItemViewModel.Reset();
+	if(DetailsPanel.IsValid())
+	{
+		if(HierarchyItem.IsValid() && HierarchyItem->AllowEditingInDetailsPanel())
+		{
+			// when se select a section, and the previous item selection is no longer available due to it, we would get a selection refresh next tick
+			// to wipe out the current selection. We want to avoid that, so we manually clear the selected items in that case.
+			if(HierarchyItem->GetData()->IsA<UHierarchySection>())
+			{
+				HierarchyTreeView->ClearSelection();
+			}
+
+			// we clear the selection of the other tree view to avoid 
+			if(bFromHierarchy)
+			{
+				SourceTreeView->ClearSelection();
+			}
+			else
+			{
+				HierarchyTreeView->ClearSelection();
+			}
+			
+			UObject* DataForEditing = HierarchyItem->GetDataForEditing();
+			DataForEditing->SetFlags(RF_Transactional);
+
+			// we make sure the object we are editing is transactional
+			DetailsPanel->SetObject(DataForEditing);
+			SelectedDetailsPanelItemViewModel = HierarchyItem;
+		}
+		else
+		{
+			SelectedDetailsPanelItemViewModel.Reset();
+			DetailsPanel->SetObject(nullptr);
+		}
+	}
+	
+	if(DetailsPanel.IsValid() && OnGenerateCustomDetailsPanelNameWidget.IsBound() && SelectedDetailsPanelItemViewModel.IsValid())
+	{
+		TSharedRef<SWidget> NameWidget = SNew(SHorizontalBox)
+		+ SHorizontalBox::Slot()
+		.AutoWidth()
+		.Padding(2.f)
+		.HAlign(HAlign_Center)
+		.VAlign(VAlign_Center)
+		[
+			SNew(SImage)
+			.Image(FAppStyle::GetBrush("Icons.Lock"))
+			.Visibility(SelectedDetailsPanelItemViewModel.Pin()->IsEditableByUser().bCanPerform ? EVisibility::Collapsed : EVisibility::Visible)
+			.ToolTipText(SelectedDetailsPanelItemViewModel.Pin()->IsEditableByUser().CanPerformMessage)
+		]
+		+ SHorizontalBox::Slot()
+		.AutoWidth()
+		.Padding(2.f)
+		[
+			OnGenerateCustomDetailsPanelNameWidget.Execute(SelectedDetailsPanelItemViewModel.Pin())
+		];
+		
+		DetailsPanel->SetNameAreaCustomContent(NameWidget);
+	}
+}
+
+void SDataHierarchyEditor::RunSourceSearch()
+{
+	if(!SourceSearchBox->GetText().IsEmpty())
+	{
+		OnSourceSearchTextChanged(SourceSearchBox->GetText());
+	}
+}
+
+void SDataHierarchyEditor::OnSourceSearchTextChanged(const FText& Text)
+{
+	SourceSearchResults.Empty();
+	FocusedSearchResult.Reset();
+	SourceTreeView->ClearSelection();
+
+	if(!Text.IsEmpty())
+	{
+		FString TextAsString = Text.ToString();;
+
+		TArray<FSearchItem> SearchItems;
+		GenerateSearchItems(SourceRootViewModel.ToSharedRef(), {}, SearchItems);
+		
+		for(const FSearchItem& SearchItem : SearchItems)
+		{
+			for(const FString& SearchTerm : SearchItem.GetEntry()->GetSearchTerms())
+			{
+				if(SearchTerm.Contains(TextAsString))
+				{
+					SourceSearchResults.Add(SearchItem);
+				}
+			}
+		}
+
+		ExpandSourceSearchResults();
+		SelectNextSourceSearchResult();
+	}
+	else
+	{
+		SourceTreeView->ClearExpandedItems();
+	}
+}
+
+void SDataHierarchyEditor::OnSourceSearchTextCommitted(const FText& Text, ETextCommit::Type CommitType)
+{
+	bool bIsShiftDown = FSlateApplication::Get().GetModifierKeys().IsShiftDown();
+	if(CommitType == ETextCommit::OnEnter)
+	{
+		if(bIsShiftDown == false)
+		{
+			SelectNextSourceSearchResult();
+		}
+		else
+		{
+			SelectPreviousSourceSearchResult();
+		}
+	}
+}
+
+void SDataHierarchyEditor::OnSearchButtonClicked(SSearchBox::SearchDirection SearchDirection)
+{
+	if(SearchDirection == SSearchBox::Next)
+	{
+		SelectNextSourceSearchResult();
+	}
+	else
+	{
+		SelectPreviousSourceSearchResult();
+	}
+}
+
+void SDataHierarchyEditor::GenerateSearchItems(TSharedRef<FHierarchyElementViewModel> Root, TArray<TSharedPtr<FHierarchyElementViewModel>> ParentChain, TArray<FSearchItem>& OutSearchItems)
+{
+	const TArray<TSharedPtr<FHierarchyElementViewModel>> FilteredChildren = Root->GetFilteredChildren();
+	ParentChain.Add(Root);
+	OutSearchItems.Add(FSearchItem{ParentChain});
+	for(TSharedPtr<FHierarchyElementViewModel> Child : FilteredChildren)
+	{
+		GenerateSearchItems(Child.ToSharedRef(), ParentChain, OutSearchItems);
+	}
+}
+
+void SDataHierarchyEditor::ExpandSourceSearchResults()
+{
+	SourceTreeView->ClearExpandedItems();
+
+	for(FSearchItem& SearchResult : SourceSearchResults)
+	{
+		for(TSharedPtr<FHierarchyElementViewModel>& EntryInPath : SearchResult.Path)
+		{
+			SourceTreeView->SetItemExpansion(EntryInPath, true);
+		}
+	}
+}
+
+void SDataHierarchyEditor::SelectNextSourceSearchResult()
+{
+	if(SourceSearchResults.IsEmpty())
+	{
+		return;
+	}
+	
+	if(!FocusedSearchResult.IsSet())
+	{
+		FocusedSearchResult = SourceSearchResults[0];
+	}
+	else
+	{
+		int32 CurrentSearchResultIndex = SourceSearchResults.Find(FocusedSearchResult.GetValue());
+		if(SourceSearchResults.IsValidIndex(CurrentSearchResultIndex+1))
+		{
+			FocusedSearchResult = SourceSearchResults[CurrentSearchResultIndex+1];
+		}
+		else
+		{
+			FocusedSearchResult = SourceSearchResults[0];
+		}
+	}
+
+	SourceTreeView->ClearSelection();
+	SourceTreeView->RequestScrollIntoView(FocusedSearchResult.GetValue().GetEntry());
+	SourceTreeView->SetItemSelection(FocusedSearchResult.GetValue().GetEntry(), true);
+}
+
+void SDataHierarchyEditor::SelectPreviousSourceSearchResult()
+{
+	if(SourceSearchResults.IsEmpty())
+	{
+		return;
+	}
+	
+	if(!FocusedSearchResult.IsSet())
+	{
+		FocusedSearchResult = SourceSearchResults[0];
+	}
+	else
+	{
+		int32 CurrentSearchResultIndex = SourceSearchResults.Find(FocusedSearchResult.GetValue());
+		if(SourceSearchResults.IsValidIndex(CurrentSearchResultIndex-1))
+		{
+			FocusedSearchResult = SourceSearchResults[CurrentSearchResultIndex-1];
+		}
+		else
+		{
+			FocusedSearchResult = SourceSearchResults[SourceSearchResults.Num()-1];
+		}
+	}
+
+	SourceTreeView->ClearSelection();
+	SourceTreeView->RequestScrollIntoView(FocusedSearchResult.GetValue().GetEntry());
+	SourceTreeView->SetItemSelection(FocusedSearchResult.GetValue().GetEntry(), true);
+}
+
+TOptional<SSearchBox::FSearchResultData> SDataHierarchyEditor::GetSearchResultData() const
+{
+	if(SourceSearchResults.Num() > 0)
+	{
+		SSearchBox::FSearchResultData SearchResultData;
+		SearchResultData.NumSearchResults = SourceSearchResults.Num();
+
+		if(FocusedSearchResult.IsSet())
+		{
+			// we add one just to make it look nicer as this is merely for cosmetic purposes
+			SearchResultData.CurrentSearchResultIndex = SourceSearchResults.Find(FocusedSearchResult.GetValue()) + 1;
+		}
+		else
+		{
+			SearchResultData.CurrentSearchResultIndex = INDEX_NONE;
+		}
+
+		return SearchResultData;
+	}
+
+	return TOptional<SSearchBox::FSearchResultData>();
+}
+
+FReply SDataHierarchyEditor::HandleHierarchyRootDrop(const FGeometry& Geometry, const FDragDropEvent& DragDropEvent) const
+{
+	if(TSharedPtr<FHierarchyDragDropOp> DragDropOp = DragDropEvent.GetOperationAs<FHierarchyDragDropOp>())
+	{
+		HierarchyViewModel->GetHierarchyRootViewModel()->OnDroppedOn(DragDropOp->GetDraggedElement().Pin(), EItemDropZone::OntoItem);
+		return FReply::Handled();
+	}
+
+	return FReply::Unhandled();
+}
+
+FHierarchyElementViewModel::FCanPerformActionResults SDataHierarchyEditor::CanDropOnRoot(TSharedPtr<FHierarchyElementViewModel> DraggedItem) const
+{
+	return HierarchyViewModel->GetHierarchyRootViewModel()->CanDropOnInternal(DraggedItem, EItemDropZone::OntoItem);
+}
+
+bool SDataHierarchyEditor::OnCanDropOnRoot(TSharedPtr<FDragDropOperation> DragDropOperation) const
+{
+	if(DragDropOperation->IsOfType<FHierarchyDragDropOp>())
+	{
+		TSharedPtr<FHierarchyDragDropOp> HierarchyDragDropOp = StaticCastSharedPtr<FHierarchyDragDropOp>(DragDropOperation);
+		TSharedPtr<FHierarchyElementViewModel> DraggedItem = HierarchyDragDropOp->GetDraggedElement().Pin();
+		return CanDropOnRoot(DraggedItem).bCanPerform;
+	}
+	
+	return false;
+}
+
+void SDataHierarchyEditor::OnRootDragEnter(const FDragDropEvent& DragDropEvent) const
+{
+	if(TSharedPtr<FHierarchyDragDropOp> HierarchyDragDropOp = DragDropEvent.GetOperationAs<FHierarchyDragDropOp>())
+	{
+		FHierarchyElementViewModel::FCanPerformActionResults CanPerformActionResults = CanDropOnRoot(HierarchyDragDropOp->GetDraggedElement().Pin());
+		HierarchyDragDropOp->SetDescription(CanPerformActionResults.CanPerformMessage);
+	}
+}
+
+void SDataHierarchyEditor::OnRootDragLeave(const FDragDropEvent& DragDropEvent) const
+{
+	if(TSharedPtr<FHierarchyDragDropOp> HierarchyDragDropOp = DragDropEvent.GetOperationAs<FHierarchyDragDropOp>())
+	{
+		HierarchyDragDropOp->SetDescription(FText::GetEmpty());
+	}
+}
+
+FSlateColor SDataHierarchyEditor::GetRootIconColor() const
+{
+	if(FSlateApplication::Get().IsDragDropping())
+	{
+		if(FSlateApplication::Get().GetDragDroppingContent()->IsOfType<FHierarchyDragDropOp>())
+		{
+			TSharedPtr<FHierarchyDragDropOp> HierarchyDragDropOp = StaticCastSharedPtr<FHierarchyDragDropOp>(FSlateApplication::Get().GetDragDroppingContent());
+			if(CanDropOnRoot(HierarchyDragDropOp->GetDraggedElement().Pin()).bCanPerform)
+			{
+				return FLinearColor(0.8f, 0.8f, 0.8f, 0.8f);
+			}
+		}
+	}
+
+	return FLinearColor(0.2f, 0.2f, 0.2f, 0.5f);
+}
+
+void SDataHierarchyEditor::NotifyPostChange(const FPropertyChangedEvent& PropertyChangedEvent,	FProperty* PropertyThatChanged)
+{
+	HierarchyViewModel->OnHierarchyPropertiesChanged().Broadcast();
+}
+
+void SSectionDragDropTarget::Construct(const FArguments& InArgs,TSharedPtr<FHierarchySectionViewModel> InOwningSection, EItemDropZone InItemDropZone)
+{
+	OwningSection = InOwningSection;
+	DropZone = InItemDropZone;
+	
+	SDropTarget::Construct(InArgs._DropTargetArgs);
+}
+
+FReply SSectionDragDropTarget::OnDragOver(const FGeometry& MyGeometry, const FDragDropEvent& DragDropEvent)
+{
+	if(TSharedPtr<FHierarchyDragDropOp> HierarchyDragDropOp = DragDropEvent.GetOperationAs<FHierarchyDragDropOp>())
+	{
+		TSharedPtr<FHierarchyElementViewModel> DraggedItem = HierarchyDragDropOp->GetDraggedElement().Pin();
+		FHierarchyElementViewModel::FCanPerformActionResults Results = OwningSection->CanDropOn(DraggedItem, DropZone);
+		HierarchyDragDropOp->SetDescription(Results.CanPerformMessage);
+	}
+
+	return SDropTarget::OnDragOver(MyGeometry, DragDropEvent);
+}
+
+void SSectionDragDropTarget::OnDragLeave(const FDragDropEvent& DragDropEvent)
+{
+	TSharedPtr<FHierarchyDragDropOp> HierarchyDragDropOp = DragDropEvent.GetOperationAs<FHierarchyDragDropOp>();
+
+	if(HierarchyDragDropOp.IsValid())
+	{
+		HierarchyDragDropOp->SetDescription(FText::GetEmpty());
+	}
+	
+	SDropTarget::OnDragLeave(DragDropEvent);
+}
+
+#undef LOCTEXT_NAMESPACE
diff --git a/Engine/Source/Editor/DataHierarchyEditor/Public/DataHierarchyViewModelBase.h b/Engine/Source/Editor/DataHierarchyEditor/Public/DataHierarchyViewModelBase.h
new file mode 100644
--- /dev/null
+++ b/Engine/Source/Editor/DataHierarchyEditor/Public/DataHierarchyViewModelBase.h
@@ -0,0 +1,968 @@
+// Copyright Epic Games, Inc. All Rights Reserved.
+
+#pragma once
+
+#include "CoreMinimal.h"
+#include "PropertyEditorDelegates.h"
+#include "ToolMenuSection.h"
+#include "IPropertyRowGenerator.h"
+#include "Misc/TransactionObjectEvent.h"
+#include "EditorUndoClient.h"
+#include "Widgets/Views/STableRow.h"
+#include "TickableEditorObject.h"
+#include "Engine/TimerHandle.h"
+#include "Templates/SubclassOf.h"
+#include "UObject/Package.h"
+#include "UObject/UObjectGlobals.h"
+#include "DataHierarchyViewModelBase.generated.h"
+
+/** HIERARCHY EDITOR
+ *	The hierarchy editor is a generic tool to organize and structure all kinds of data.
+ *	It inherently supports sections, categories, and items. You can add your own items and customize how they are visualized.
+ *	Each hierarchy element is a UObject, and some hierarchy elements will represent externally owned data.
+ *	For example, categories and sections defined within the hierarchy are also owned by the hierarchy, but an item might represent a parameter defined elsewhere.
+ *
+ *  Each hierarchy element is pure data and should not reference externally owned data that could become invalid.
+ *  To define per-element rules, each hierarchy element gets assigned one view model.
+ *  
+ *	To use the Hierarchy Editor, you need multiple things:
+ *	1) A UHierarchyRoot object that the Hierarchy Editor uses to store the created hierarchy
+ *	2) A UDataHierarchyViewModelBase-derived object that defines core hierarchy rules. This is the main object responsible for configuring your hierarchy.
+ *	-- The UDataHierarchyViewModelBase-derived class has multiple virtual functions you need to override. The key functions are:
+ *	--- a) GetHierarchyRoot(), pointing to the UHierarchyRoot object you created in 1)
+ *	--- b) PrepareSourceItems(UHierarchyRoot* SourceRoot, ...), which you need to use to populate the list of elements to be organized.
+ *	------ These are transient items and a transient root. You can also add sections and categories here if you wish to structure the source element view.
+ *	------ To add a child, call NewObject with ParentElement as the outer (in a flat list the SourceRoot).
+ *	------ Then use ParentElement->GetChildrenMutable().Add(NewElement).
+ *	------ This way, you can write your own initialization functions
+ *	------ You can also add categories and sections to the SourceRoot
+ *	--- c) Optionally but likely: CreateCustomViewModelForData(UHierarchyElement* Element, ...), which is used to create and assign non-default view models for each hierarchy element
+ *	------ You don't need to implement this, but any slightly advanced use case that requires modification to hierarchy logic will need to implement this
+ *	3) An SHierarchyEditor widget, which takes in the UDataHierarchyViewModelBase-derived object you created in 2). This way, the hierarchy editor knows where to store data, and how to initialize.
+ *	--- Additionally, it has arguments such as 'OnGenerateRowContentWidget', which gives you an ElementViewModel.
+ *	--- After determining the type, you can use this to define widgets for each hierarchy element.
+ *
+ *	Tips and tricks:
+ *	1) Each hierarchy element can have a FHierarchyElementIdentity consisting of guid(s) and/or name(s).
+ *	--- This can be used to uniquely identify an element, navigate to it and so on.
+ *	2) To deal with automated cleanup of stale hierarchy elements that represent external data, you can set a UHierarchyDataRefreshContext-derived object on the UDataHierarchyViewModelBase object.
+ *	--- The idea is to create a new derived class with objects and properties of the system that uses the Hierarchy Editor.
+ *	--- For example, in a graph based system, the DataRefreshContext object can point to the graph.
+ *	--- In an FHierarchyElementViewModel's 'DoesExternalDataStillExist' function,
+ *	--- you can access the DataRefreshContext to determine whether your externally represented hierarchy elements should still exist.
+ *	--- If not, it is automatically deleted.
+ *	3) In the details panel, You can edit the hierarchy elements themselves, or external objects 
+ *
+ *  To make use of the created hierarchy, you access the UHierarchyRoot object you created in 1), and query it for its children, sections etc.
+ *  The HierarchyEditor does not define how to use the created hierarchy data in your own UI. It only lets you structure and edit data.
+ *  How you use it from a 'data consumption' point of view is up to you.
+ */
+
+struct FHierarchyRootViewModel;
+/** This struct is used to identify a given hierarchy element and can be based on guids and/or names.
+ *  This is particularly useful when a hierarchy element represents an object or a property that is not owned by the hierarchy itself.
+ *  
+ *  
+ */
+USTRUCT()
+struct FHierarchyElementIdentity
+{
+	GENERATED_BODY()
+
+	FHierarchyElementIdentity() {}
+	FHierarchyElementIdentity(TArray<FGuid> InGuids, TArray<FName> InNames) : Guids(InGuids), Names(InNames) {}
+	
+	/** An array of guids that have to be satisfied in order to match. */
+	UPROPERTY()
+	TArray<FGuid> Guids;
+
+	/** Optionally, an array of names can be specified in place of guids. If guids & names are present, guids have to be satisfied first, then names. */
+	UPROPERTY()
+	TArray<FName> Names;
+
+	bool IsValid() const
+	{
+		return Guids.Num() > 0 || Names.Num() > 0;
+	}
+	
+	bool operator==(const FHierarchyElementIdentity& OtherIdentity) const
+	{
+		if(Guids.Num() != OtherIdentity.Guids.Num() || Names.Num() != OtherIdentity.Names.Num())
+		{
+			return false;
+		}
+
+		for(int32 GuidIndex = 0; GuidIndex < Guids.Num(); GuidIndex++)
+		{
+			if(Guids[GuidIndex] != OtherIdentity.Guids[GuidIndex])
+			{
+				return false;
+			}
+		}
+
+		for(int32 NameIndex = 0; NameIndex < Names.Num(); NameIndex++)
+		{
+			if(!Names[NameIndex].IsEqual(OtherIdentity.Names[NameIndex]))
+			{
+				return false;
+			}
+		}
+
+		return true;
+	}
+
+	bool operator!=(const FHierarchyElementIdentity& OtherIdentity) const
+	{
+		return !(*this == OtherIdentity);
+	}
+};
+
+FORCEINLINE uint32 GetTypeHash(const FHierarchyElementIdentity& Identity)
+{
+	uint32 Hash = 0;
+	
+	for(const FGuid& Guid : Identity.Guids)
+	{
+		Hash = HashCombine(Hash, GetTypeHash(Guid));
+	}
+	
+	for(const FName& Name : Identity.Names)
+	{
+		Hash = HashCombine(Hash, GetTypeHash(Name));
+	}
+	
+	return Hash;
+}
+
+/** A base class that is used to refresh data that represents external data. Inherit from this class if you need more context data. */
+UCLASS(Transient)
+class DATAHIERARCHYEDITOR_API UHierarchyDataRefreshContext : public UObject
+{
+	GENERATED_BODY()
+};
+
+UCLASS()
+class DATAHIERARCHYEDITOR_API UHierarchyElement : public UObject
+{
+	GENERATED_BODY()
+
+public:
+	UHierarchyElement() { Identity.Guids.Add(FGuid::NewGuid()); }
+	virtual ~UHierarchyElement() override {}
+
+	TArray<TObjectPtr<UHierarchyElement>>& GetChildrenMutable() { return Children; }
+	const TArray<UHierarchyElement*>& GetChildren() const;
+
+	template<class ChildClass>
+	ChildClass* AddChild();
+
+	UHierarchyElement* FindChildWithIdentity(FHierarchyElementIdentity ChildIdentity, bool bSearchRecursively = false);
+
+	UHierarchyElement* CopyAndAddItemAsChild(const UHierarchyElement& ItemToCopy);
+	UHierarchyElement* CopyAndAddItemUnderParentIdentity(const UHierarchyElement& ItemToCopy, FHierarchyElementIdentity ParentIdentity);
+	
+	/** Remove a child with a given identity. Can be searched recursively. This function operates under the assumption there will be only one item with a given identity. */
+	bool RemoveChildWithIdentity(FHierarchyElementIdentity ChildIdentity, bool bSearchRecursively = false);
+	
+	template<class ChildClass>
+	bool DoesOneChildExist(bool bRecursive = false) const;
+
+	template<class ChildClass>
+	TArray<ChildClass*> GetChildrenOfType(TArray<ChildClass*>& Out, bool bRecursive = false) const;
+
+	template<class PREDICATE_CLASS>
+	void SortChildren(const PREDICATE_CLASS& Predicate, bool bRecursive = false);
+	
+	virtual FString ToString() const { return GetName(); }
+	FText ToText() const { return FText::FromString(ToString()); }
+	
+	/** An identity can be optionally set to create a mapping from previously existing guids or names to hierarchy items that represent them. */
+	void SetIdentity(FHierarchyElementIdentity InIdentity) { Identity = InIdentity; }
+	FHierarchyElementIdentity GetPersistentIdentity() const { return Identity; }
+	TArray<FHierarchyElementIdentity> GetParentIdentities() const;
+	
+	/** Overridden modify method to also mark all children as modified */
+	virtual bool Modify(bool bAlwaysMarkDirty = true) override;
+
+protected:
+	virtual void PostLoad() override;
+
+	UPROPERTY()
+	TArray<TObjectPtr<UHierarchyElement>> Children;
+
+	UPROPERTY()
+	FHierarchyElementIdentity Identity;
+	
+	/** An optional guid; can be used if hierarchy items represent outside items */
+	UPROPERTY()
+	FGuid Guid_DEPRECATED;
+};
+
+template <class ChildClass>
+ChildClass* UHierarchyElement::AddChild()
+{
+	ChildClass* NewChild = NewObject<ChildClass>(this);
+	GetChildrenMutable().Add(NewChild);
+
+	return NewChild;
+}
+
+template <class ChildClass>
+bool UHierarchyElement::DoesOneChildExist(bool bRecursive) const
+{
+	for(UHierarchyElement* ChildElement : Children)
+	{
+		if(ChildElement->IsA<ChildClass>())
+		{
+			return true;
+		}
+	}
+
+	if(bRecursive)
+	{
+		for(UHierarchyElement* ChildElement : Children)
+		{
+			if(ChildElement->DoesOneChildExist<ChildClass>(bRecursive))
+			{
+				return true;
+			}
+		}
+	}
+
+	return false;
+}
+
+template <class ChildClass>
+TArray<ChildClass*> UHierarchyElement::GetChildrenOfType(TArray<ChildClass*>& Out, bool bRecursive) const
+{
+	for(UHierarchyElement* ChildElement : Children)
+	{
+		if(ChildElement->IsA<ChildClass>())
+		{
+			Out.Add(Cast<ChildClass>(ChildElement));
+		}
+	}
+
+	if(bRecursive)
+	{
+		for(UHierarchyElement* ChildElement : Children)
+		{
+			ChildElement->GetChildrenOfType<ChildClass>(Out, bRecursive);
+		}
+	}
+
+	return Out;
+}
+
+template <class PREDICATE_CLASS>
+void UHierarchyElement::SortChildren(const PREDICATE_CLASS& Predicate, bool bRecursive)
+{
+	Children.Sort(Predicate);
+
+	if(bRecursive)
+	{
+		for(TObjectPtr<UHierarchyElement> ChildElement : Children)
+		{
+			ChildElement->SortChildren(Predicate, bRecursive);
+		}
+	}
+}
+
+/** A minimal implementation of a section. */
+UCLASS()
+class DATAHIERARCHYEDITOR_API UHierarchySection : public UHierarchyElement
+{
+	GENERATED_BODY()
+
+public:
+	UHierarchySection() {}
+
+	void SetSectionName(FName InSectionName) { Section = InSectionName; }
+	FName GetSectionName() const { return Section; }
+	
+	void SetSectionNameAsText(const FText& Text);
+	FText GetSectionNameAsText() const { return FText::FromName(Section); }
+
+	void SetTooltip(const FText& InTooltip) { Tooltip = InTooltip; }
+	FText GetTooltip() const { return Tooltip; }
+
+	virtual FString ToString() const override { return Section.ToString(); }
+private:
+	UPROPERTY()
+	FName Section;
+
+	/** The tooltip used when the user is hovering this section */
+	UPROPERTY(EditAnywhere, Category = "Section", meta = (MultiLine = "true"))
+	FText Tooltip;
+};
+
+/** UHierarchyRoot is used as the main object for serialization purposes, and a transient root is created automatically by the widget to populate the source list of items. */
+UCLASS()
+class DATAHIERARCHYEDITOR_API UHierarchyRoot : public UHierarchyElement
+{
+	GENERATED_BODY()
+public:
+	UHierarchyRoot() {}
+	virtual ~UHierarchyRoot() override {}
+
+	const TArray<UHierarchySection*>& GetSectionData() const { return Sections; }
+	TArray<TObjectPtr<UHierarchySection>>& GetSectionDataMutable() { return Sections; }
+
+	TSet<FName> GetSections() const;
+	int32 GetSectionIndex(UHierarchySection* Section) const;
+
+	UHierarchySection* AddSection(FText InNewSectionName, int32 InsertIndex = INDEX_NONE, TSubclassOf<UHierarchySection> SectionClass = UHierarchySection::StaticClass());
+	UHierarchySection* FindSectionByIdentity(FHierarchyElementIdentity SectionIdentity);
+	/** This will copy the section element itself */
+	void DuplicateSectionFromOtherRoot(const UHierarchySection& SectionToCopy);
+	void RemoveSection(FText SectionName);
+	void RemoveSectionByIdentity(FHierarchyElementIdentity SectionIdentity);
+	
+	virtual bool Modify(bool bAlwaysMarkDirty = true) override;
+	virtual void Serialize(FStructuredArchive::FRecord Record) override;
+
+	virtual void EmptyAllData();
+protected:
+
+	UPROPERTY()
+	TArray<TObjectPtr<UHierarchySection>> Sections;
+};
+
+/** A minimal implementation of an item. Inherit from this and add your own properties. */
+UCLASS(MinimalAPI)
+class UHierarchyItem : public UHierarchyElement
+{
+	GENERATED_BODY()
+public:
+	UHierarchyItem() {}
+	virtual ~UHierarchyItem() override {}
+};
+
+/** A category, potentially pointing at the section it belongs to. Only top-level categories can belong to sections by default.
+ *  Inherit from this to add your own properties.  */
+UCLASS()
+class DATAHIERARCHYEDITOR_API UHierarchyCategory : public UHierarchyElement
+{
+	GENERATED_BODY()
+public:
+	UHierarchyCategory() {}
+	UHierarchyCategory(FName InCategory) : Category(InCategory) {}
+	
+	void SetCategoryName(FName NewCategory) { Category = NewCategory; }
+	FName GetCategoryName() const { return Category; }
+
+	FText GetCategoryAsText() const { return FText::FromName(Category); }
+	FText GetTooltip() const { return Tooltip; }
+
+	void SetSection(UHierarchySection* InSection) { Section = InSection; }
+	const UHierarchySection* GetSection() const { return Section; }
+
+	virtual FString ToString() const override { return Category.ToString(); }
+
+	/** Since the category points to a section object, during merge or copy paste etc. it is possible the section pointer will point at a section from another root.
+	 *  We fix this up by looking through our available sections and match up via persistent identity.
+	 *  This function expects the correct section with the same identity to exist already at the root level
+	 */
+	void FixupSectionLinkage();
+private:
+	UPROPERTY()
+	FName Category;
+
+	/** The tooltip used when the user is hovering this category */
+	UPROPERTY(EditAnywhere, Category = "Category", meta = (MultiLine = "true"))
+	FText Tooltip;
+
+	UPROPERTY()
+	TObjectPtr<UHierarchySection> Section = nullptr;
+};
+
+struct FHierarchyElementViewModel;
+struct FHierarchySectionViewModel;
+
+/** Inherit from this to allow UI customization for your drag & drop operation by overriding CreateCustomDecorator. */
+class DATAHIERARCHYEDITOR_API FHierarchyDragDropOp : public FDragDropOperation
+{
+public:
+	DRAG_DROP_OPERATOR_TYPE(FHierarchyDragDropOp, FDragDropOperation)
+
+	FHierarchyDragDropOp(TSharedPtr<FHierarchyElementViewModel> InDraggedElementViewModel);
+
+	virtual void Construct() override { FDragDropOperation::Construct(); }
+	virtual TSharedPtr<SWidget> GetDefaultDecorator() const override final;
+	
+	/** Override this custom decorator function to provide custom widget visuals. If not specified, you can still use Label and Description. */
+	virtual TSharedRef<SWidget> CreateCustomDecorator() const { return SNullWidget::NullWidget; }
+	
+	TWeakPtr<FHierarchyElementViewModel> GetDraggedElement() { return DraggedElement; }
+	
+	void SetLabel(FText InText) { Label = InText; }
+	FText GetLabel() const { return Label; }
+
+	void SetDescription(FText InText) { Description = InText; }
+	FText GetDescription() const { return Description; }
+
+	void SetFromSourceList(bool bInFromSourceList) { bFromSourceList = bInFromSourceList; }
+	bool GetIsFromSourceList() const { return bFromSourceList; }
+protected:
+	TWeakPtr<FHierarchyElementViewModel> DraggedElement;
+	
+	/** Label will be displayed if no custom decorator has been specified. */
+	FText Label;
+	/** Useful for runtime tweaking of the tooltip based on what we are hovering. Always displayed if not-empty */
+	FText Description;
+	/** If the drag drop op is from the source list, we can further customize the actions */
+	bool bFromSourceList = false;
+};
+
+UCLASS(BlueprintType)
+class DATAHIERARCHYEDITOR_API UHierarchyMenuContext : public UObject
+{
+	GENERATED_BODY()
+
+public:
+	TWeakObjectPtr<class UDataHierarchyViewModelBase> HierarchyViewModel;
+	TArray<TSharedPtr<FHierarchyElementViewModel>> MenuHierarchyElements;
+};
+
+/** The main controller class for the SHierarchyEditor widget. Defines core hierarchy rules.
+ *  Inherit from this and override the required virtual functions, instantiate an object, Initialize it and pass it to the SHierarchyEditor widget. */
+UCLASS(Abstract, Transient)
+class DATAHIERARCHYEDITOR_API UDataHierarchyViewModelBase : public UObject, public FSelfRegisteringEditorUndoClient
+{
+public:
+	DECLARE_MULTICAST_DELEGATE(FOnHierarchyChanged)
+	DECLARE_MULTICAST_DELEGATE(FOnHierarchyPropertiesChanged)
+	DECLARE_DELEGATE_OneParam(FOnSectionActivated, TSharedPtr<FHierarchySectionViewModel> Section)
+	DECLARE_DELEGATE_OneParam(FOnElementAdded, TSharedPtr<FHierarchyElementViewModel> AddedElement)
+	DECLARE_DELEGATE_OneParam(FOnRefreshViewRequested, bool bForceFullRefresh)
+	DECLARE_DELEGATE_OneParam(FOnNavigateToElementIdentityInHierarchyRequested, FHierarchyElementIdentity Identity)
+	DECLARE_DELEGATE_OneParam(FOnNavigateToElementInHierarchyRequested, TSharedPtr<FHierarchyElementViewModel> ElementViewModel)
+
+	GENERATED_BODY()
+
+	UDataHierarchyViewModelBase();
+	virtual ~UDataHierarchyViewModelBase() override;
+
+	/** Initialize is called automatically for you, but it is recommended to call it manually after creating the HierarchyViewModel in your own Initialize function. This lets you access external data. */
+	void Initialize();
+	/** Call Finalize manually when you no longer need the HierarchyViewModel. */
+	void Finalize();
+	
+	bool IsInitialized() const { return bIsInitialized; }
+	bool IsFinalized() const { return bIsFinalized; }
+	bool IsValid() const { return IsInitialized() && !IsFinalized(); }
+
+	FName GetContextMenuName () const;
+
+	/** Creates view model hierarchy elements. To create custom view models, override CreateCustomViewModelForElement. */
+	TSharedPtr<FHierarchyElementViewModel> CreateViewModelForElement(UHierarchyElement* Element, TSharedPtr<FHierarchyElementViewModel> Parent);
+	/** Get the root view model associated with the hierarchy. */
+	TSharedPtr<struct FHierarchyRootViewModel> GetHierarchyRootViewModel() const { return HierarchyRootViewModel; }
+	/** Hierarchy items reflect the already edited hierarchy. This should generally be constructed from persistent serialized data. */
+	const TArray<TSharedPtr<FHierarchyElementViewModel>>& GetHierarchyItems() const;
+
+public:
+	/** The hierarchy root the widget is editing. This should point to persistent data stored somewhere else as the serialized root of the hierarchy. */
+	virtual UHierarchyRoot* GetHierarchyRoot() const PURE_VIRTUAL(UDataHierarchyViewModelBase::GetHierarchyRoot, return nullptr;);
+	/** The outer for the transient source root creation can be overridden. */
+	virtual UObject* GetOuterForSourceRoot() const { return GetTransientPackage(); }
+	/** Prepares the items we want to create a hierarchy for. Primary purpose is to add children to the source root to gather the items to display in the source panel.
+	 * The root view model is also given as a way to forcefully sync view models to access additional functionality, if needed */
+	virtual void PrepareSourceItems(UHierarchyRoot* SourceRoot, TSharedPtr<FHierarchyRootViewModel> SourceRootViewModel) PURE_VIRTUAL(UDataHierarchyViewModelBase::PrepareSourceItems,);
+	/** The class used for creating categories. You can subclass UHierarchyCategory to add new properties. */
+	virtual TSubclassOf<UHierarchyCategory> GetCategoryDataClass() const;
+	/** The class used for creating sections. You can subclass UHierarchySection to add new properties. */
+	virtual TSubclassOf<UHierarchySection> GetSectionDataClass() const;
+	/** Function to implement drag drop ops. FHierarchyDragDropOp is a default implementation for a single hierarchy element. Inherit from it and override CreateCustomDecorator for custom UI. */
+	virtual TSharedRef<FHierarchyDragDropOp> CreateDragDropOp(TSharedRef<FHierarchyElementViewModel> Item);
+	/** This needs to return true if you want the details panel to show up. */
+	virtual bool SupportsDetailsPanel() { return true; }
+	/** Overriding this allows to define details panel instance customizations for specific UClasses */
+	virtual TArray<TTuple<UClass*, FOnGetDetailCustomizationInstance>> GetInstanceCustomizations() { return {}; }
+
+protected:
+	/** Additional commands can be specified overriding the SetupCommands function. */
+	virtual void SetupCommands() {}
+private:
+	/** Lets you add some additional logic to the Initialize function. */
+	virtual void InitializeInternal() {}
+	virtual void FinalizeInternal() {}
+
+	/** This function is used to determine custom view models for Hierarchy Elements. Called by CreateViewModelForElement. */
+	virtual TSharedPtr<FHierarchyElementViewModel> CreateCustomViewModelForElement(UHierarchyElement* Element, TSharedPtr<FHierarchyElementViewModel> Parent) { return nullptr; }
+	/** Override this to add custom context menu options. Defaults to common actions such as finding, renaming and deleting.
+	 * The hierarchy elements to generate this menu for are contained in a UHierarchyMenuContext object accessible through the ToolMenu. */
+	virtual void GenerateDynamicContextMenuInternal(UToolMenu* DynamicToolMenu) const;
+public:
+	const UHierarchyDataRefreshContext* GetRefreshContext() const { return RefreshContext; }
+	/** Set the refresh context to easily allow Hierarchy Elements representing external data to access whether the external data still exists. */
+	void SetRefreshContext(UHierarchyDataRefreshContext* InContext) { RefreshContext = InContext; }
+	
+	UHierarchyElement* AddElementUnderRoot(TSubclassOf<UHierarchyElement> NewChildClass, FHierarchyElementIdentity ChildIdentity);
+	void AddCategory(TSharedPtr<FHierarchyElementViewModel> CategoryParent = nullptr) const;
+	void AddSection() const;
+
+	void DeleteElementWithIdentity(FHierarchyElementIdentity Identity);
+	void DeleteElementsWithIdentities(TArray<FHierarchyElementIdentity> Identities);
+	void DeleteElements(TArray<TSharedPtr<FHierarchyElementViewModel>> ViewModels) const;
+
+	void NavigateToElementInHierarchy(const FHierarchyElementIdentity& HierarchyIdentity) const;
+	void NavigateToElementInHierarchy(const TSharedRef<FHierarchyElementViewModel> HierarchyElement) const;
+	
+	/** Refreshes all data and widgets */
+	void ForceFullRefresh();
+	void ForceFullRefreshOnTimer();
+	void RequestFullRefreshNextFrame();
+	
+	TSharedRef<FUICommandList> GetCommands() const { return Commands.ToSharedRef(); }
+	
+	void OnGetChildren(TSharedPtr<FHierarchyElementViewModel> Element, TArray<TSharedPtr<FHierarchyElementViewModel>>& OutChildren) const;
+	
+	void RefreshAllViews(bool bFullRefresh = false) const;
+	void RefreshSourceView(bool bFullRefresh = false) const;
+	void RefreshHierarchyView(bool bFullRefresh = false) const;
+	void RefreshSectionsView() const;
+	
+	// Delegate that call functions from SHierarchyEditor
+	FSimpleDelegate& OnRefreshSourceItemsRequested() { return RefreshSourceItemsRequestedDelegate; }
+	FOnRefreshViewRequested& OnRefreshSourceView() { return RefreshSourceViewDelegate; }
+	FOnRefreshViewRequested& OnRefreshHierarchyView() { return RefreshHierarchyWidgetDelegate; }
+	FSimpleDelegate& OnRefreshSectionsView() { return RefreshSectionsViewDelegate; }
+
+	// Delegates for external systems
+	FOnHierarchyChanged& OnHierarchyChanged() { return OnHierarchyChangedDelegate; } 
+	FOnHierarchyChanged& OnHierarchyPropertiesChanged() { return OnHierarchyPropertiesChangedDelegate; } 
+	FOnElementAdded& OnElementAdded() { return OnElementAddedDelegate; }
+	FOnRefreshViewRequested& OnRefreshViewRequested() { return RefreshAllViewsRequestedDelegate; }
+	FOnNavigateToElementIdentityInHierarchyRequested& OnNavigateToElementIdentityInHierarchyRequested() { return OnNavigateToElementIdentityInHierarchyRequestedDelegate; }
+	FOnNavigateToElementInHierarchyRequested& OnNavigateToElementInHierarchyRequested() { return OnNavigateToElementInHierarchyRequestedDelegate; }
+	FSimpleDelegate& OnInitialized() { return OnInitializedDelegate; }
+	
+	// Sections
+	void SetActiveHierarchySection(TSharedPtr<struct FHierarchySectionViewModel>);
+	TSharedPtr<FHierarchySectionViewModel> GetActiveHierarchySection() const;
+	UHierarchySection* GetActiveHierarchySectionData() const;
+	bool IsHierarchySectionActive(const UHierarchySection* Section) const;
+	FOnSectionActivated& OnHierarchySectionActivated() { return OnHierarchySectionActivatedDelegate; }
+
+	FString OnElementToStringDebug(TSharedPtr<FHierarchyElementViewModel> ElementViewModel) const;
+
+protected:	
+	virtual void PostUndo(bool bSuccess) override;
+	virtual void PostRedo(bool bSuccess) override;
+	virtual bool MatchesContext(const FTransactionContext& InContext, const TArray<TPair<UObject*, FTransactionObjectEvent>>& TransactionObjectContexts) const override;
+private:
+	bool FilterForHierarchySection(TSharedPtr<const FHierarchyElementViewModel> ViewModel) const;
+	bool FilterForUncategorizedRootItemsInAllSection(TSharedPtr<const FHierarchyElementViewModel> ViewModel) const;
+
+	static void GenerateDynamicContextMenu(UToolMenu* ToolMenu);
+
+	void ToolMenuRequestRename(const FToolMenuContext& Context) const;
+	bool ToolMenuCanRequestRename(const FToolMenuContext& Context) const;
+
+	void ToolMenuDelete(const FToolMenuContext& Context) const;
+	bool ToolMenuCanDelete(const FToolMenuContext& Context) const;
+
+	void ToolMenuNavigateTo(const FToolMenuContext& Context) const;
+	bool ToolMenuCanNavigateTo(const FToolMenuContext& Context) const;
+protected:
+	UPROPERTY()
+	TObjectPtr<UHierarchyRoot> HierarchyRoot;
+	
+	TSharedPtr<struct FHierarchyRootViewModel> HierarchyRootViewModel;
+
+	TWeakPtr<struct FHierarchySectionViewModel> ActiveHierarchySection;
+
+	TSharedPtr<FUICommandList> Commands;
+
+	UPROPERTY(Transient)
+	TObjectPtr<UHierarchyDataRefreshContext> RefreshContext = nullptr;
+	
+protected:
+	// delegate collection to call UI functions
+	FSimpleDelegate RefreshSourceItemsRequestedDelegate;
+	FOnRefreshViewRequested RefreshAllViewsRequestedDelegate;
+	FOnRefreshViewRequested RefreshSourceViewDelegate;
+	FOnRefreshViewRequested RefreshHierarchyWidgetDelegate;
+	FSimpleDelegate RefreshSectionsViewDelegate;
+	FOnNavigateToElementIdentityInHierarchyRequested OnNavigateToElementIdentityInHierarchyRequestedDelegate;
+	FOnNavigateToElementInHierarchyRequested OnNavigateToElementInHierarchyRequestedDelegate;
+	
+	FOnElementAdded OnElementAddedDelegate;
+	FOnSectionActivated OnHierarchySectionActivatedDelegate;
+	FOnSectionActivated OnSourceSectionActivatedDelegate;
+	FOnHierarchyChanged OnHierarchyChangedDelegate;
+	FOnHierarchyPropertiesChanged OnHierarchyPropertiesChangedDelegate;
+	
+	FSimpleDelegate OnInitializedDelegate;
+
+	FTimerHandle FullRefreshNextFrameHandle;
+
+private:
+	UPROPERTY(Transient)
+	bool bIsInitialized = false;
+	
+	UPROPERTY(Transient)
+	bool bIsFinalized = false;
+};
+
+/** The base view model for all elements in the hierarchy. There are four base view models inheriting from this; for roots, items, categories, and sections.
+ *  When creating a new view model, you should inherit from one of those four base view models.
+ */
+struct DATAHIERARCHYEDITOR_API FHierarchyElementViewModel : TSharedFromThis<FHierarchyElementViewModel>, public FTickableEditorObject
+{
+	DECLARE_DELEGATE(FOnSynced)
+	DECLARE_DELEGATE_RetVal_OneParam(bool, FOnFilterChild, const TSharedPtr<const FHierarchyElementViewModel> Child);
+	DECLARE_DELEGATE_OneParam(FOnChildRequestedDeletion, TSharedPtr<FHierarchyElementViewModel> Child)
+	
+	FHierarchyElementViewModel(UHierarchyElement* InElement, TSharedPtr<FHierarchyElementViewModel> InParent, TWeakObjectPtr<UDataHierarchyViewModelBase> InHierarchyViewModel, bool bInIsForHierarchy)
+		: Element(InElement)
+		, Parent(InParent)
+		, HierarchyViewModel(InHierarchyViewModel)
+		, bIsForHierarchy(bInIsForHierarchy)
+	{		
+		
+	}
+
+	/** Can be implemented for additional logic that the constructor isn't valid for. */
+	virtual void Initialize() {}
+	
+	virtual ~FHierarchyElementViewModel() override;
+	
+	UHierarchyElement* GetDataMutable() const { return Element; }
+	const UHierarchyElement* GetData() const { return Element; }
+	
+	template<class T>
+	T* GetDataMutable() const { return Cast<T>(Element); }
+	
+	template<class T>
+	const T* GetData() const { return Cast<T>(Element); }
+	
+	virtual void Tick(float DeltaTime) override;
+	virtual TStatId GetStatId() const override;
+	
+	virtual FString ToString() const { return Element->ToString(); }
+	FText ToStringAsText() const { return FText::FromString(ToString()); }
+	virtual TArray<FString> GetSearchTerms() const { return {ToString()} ;}
+
+	void RefreshChildrenData();
+	void SyncViewModelsToData();
+	
+	/** Every item view model can define its own sort order for its children. By default we put categories above items. */
+	virtual void SortChildrenData() const;
+	
+	const TArray<TSharedPtr<FHierarchyElementViewModel>>& GetChildren() const { return Children; }
+	TArray<TSharedPtr<FHierarchyElementViewModel>>& GetChildrenMutable() { return Children; }
+	const TArray<TSharedPtr<FHierarchyElementViewModel>>& GetFilteredChildren() const;
+	
+	void AddChildFilter(FOnFilterChild InFilterChild);
+	
+	template<class DataClass, class ViewModelChildClass>
+	void GetChildrenViewModelsForType(TArray<TSharedPtr<ViewModelChildClass>>& OutChildren, bool bRecursive = false);
+
+	/** Returns the hierarchy depth via number of parents above. */
+	int32 GetHierarchyDepth() const;
+	bool HasParent(TSharedPtr<FHierarchyElementViewModel> ParentCandidate, bool bRecursive = false) const;
+	
+	TSharedRef<FHierarchyElementViewModel> DuplicateToThis(TSharedPtr<FHierarchyElementViewModel> ItemToDuplicate, int32 InsertIndex = INDEX_NONE);
+	TSharedRef<FHierarchyElementViewModel> ReparentToThis(TSharedPtr<FHierarchyElementViewModel> ItemToMove, int32 InsertIndex = INDEX_NONE);
+
+	TSharedPtr<FHierarchyElementViewModel> FindViewModelForChild(UHierarchyElement* Child, bool bSearchRecursively = false) const;
+	TSharedPtr<FHierarchyElementViewModel> FindViewModelForChild(FHierarchyElementIdentity ChildIdentity, bool bSearchRecursively = false) const;
+	int32 FindIndexOfChild(UHierarchyElement* Child) const;
+	int32 FindIndexOfDataChild(TSharedPtr<FHierarchyElementViewModel> Child) const;
+	int32 FindIndexOfDataChild(UHierarchyElement* Child) const;
+
+	UHierarchyElement* AddChild(TSubclassOf<UHierarchyElement> NewChildClass, FHierarchyElementIdentity ChildIdentity);
+	
+	/** Deleting will ask the parent to delete its child */
+	void Delete();
+	void DeleteChild(TSharedPtr<FHierarchyElementViewModel> Child);
+
+	TWeakObjectPtr<UDataHierarchyViewModelBase> GetHierarchyViewModel() const { return HierarchyViewModel; }
+
+	/** Returns a set result if the item can accept a drop either above/onto/below the item.  */
+	TOptional<EItemDropZone> OnCanRowAcceptDrop(const FDragDropEvent& DragDropEvent, EItemDropZone ItemDropZone, TSharedPtr<FHierarchyElementViewModel> Item);
+	virtual FReply OnDroppedOnRow(const FDragDropEvent& DragDropEvent, EItemDropZone ItemDropZone, TSharedPtr<FHierarchyElementViewModel> Item);
+	void OnRowDragLeave(const FDragDropEvent& DragDropEvent);
+
+	struct FCanPerformActionResults
+	{
+		FCanPerformActionResults(bool bInCanPerform) : bCanPerform(bInCanPerform) {}
+		
+		bool bCanPerform = false;
+		/** A message that is used when bCanPerform is false. Will either be used in tooltips in the hierarchy editor or as popup message. */
+		FText CanPerformMessage;
+
+		bool operator==(const bool& bOther) const
+		{
+			return bCanPerform == bOther;
+		}
+
+		bool operator!=(const bool& bOther) const
+		{
+			return !(*this==bOther);
+		}
+	};
+
+	/** Should return true if properties are supposed to be editable & needs to be true if typical operations should work on it (renaming, dragging, deleting etc.) */
+	virtual FCanPerformActionResults IsEditableByUser() { return FCanPerformActionResults(false); }
+	
+	/** Needs to be true in order to allow drag & drop operations to parent items to this item */
+	virtual bool CanHaveChildren() const { return false; }
+	
+	/** Should return true if an item should be draggable. An uneditable item can not be dragged even if CanDragInternal returns true. */
+	FCanPerformActionResults CanDrag();
+	
+	/** Returns true if renamable */
+	bool CanRename() { return IsEditableByUser().bCanPerform && CanRenameInternal(); }
+
+	void Rename(FName NewName) { RenameInternal(NewName); HierarchyViewModel->OnHierarchyPropertiesChanged().Broadcast(); }
+
+	void RequestRename()
+	{
+		if(CanRename() && OnRequestRenameDelegate.IsBound())
+		{
+			bRenamePending = false;
+			OnRequestRenameDelegate.Execute();
+		}
+	}
+
+	void RequestRenamePending()
+	{
+		if(CanRename())
+		{
+			bRenamePending = true;
+		}
+	}
+	
+	/** Returns true if deletable */
+	bool CanDelete() { return IsEditableByUser().bCanPerform && CanDeleteInternal(); }
+
+	/** Returns true if the given item can be dropped on the given target area. */
+	FCanPerformActionResults CanDropOn(TSharedPtr<FHierarchyElementViewModel> DraggedItem, EItemDropZone ItemDropZone) { return CanDropOnInternal(DraggedItem, ItemDropZone); }
+
+	/** Gets executed when an item was dropped on this. */
+	void OnDroppedOn(TSharedPtr<FHierarchyElementViewModel> DroppedItem, EItemDropZone ItemDropZone) { OnDroppedOnInternal(DroppedItem, ItemDropZone); }
+
+	/** Determines the section this item belongs to. */
+	const UHierarchySection* GetSection() const { return GetSectionInternal(); }
+	
+	/** For data cleanup that represents external data, this needs to return true in order for live cleanup to work. */
+	virtual bool RepresentsExternalData() const { return false; }
+	/** This function determines whether a hierarchy item that represents that external data should be maintained during data refresh
+	 * Needs to be implemented if RepresentsExternalData return true.
+	 * The context object can be used to add arbitrary data. */
+	virtual bool DoesExternalDataStillExist(const UHierarchyDataRefreshContext* Context) const { return false; }
+
+	/** The UObject we display in the details panel when this item is selected. By default it's the hierarchy element the view model represents. */
+	virtual UObject* GetDataForEditing() { return Element; }
+	/** Source items are transient, which is why we don't allow editing by default.
+	 * This is useful to override if source data points at actual data to edit. */
+	virtual bool AllowEditingInDetailsPanel() const { return bIsForHierarchy; }
+
+	bool IsForHierarchy() const { return bIsForHierarchy; }
+
+	/** Override this to register dynamic context menu entries when right clicking a single hierarchy item */
+	virtual void AppendDynamicContextMenuForSingleElement(UToolMenu* ToolMenu) {}
+	
+	FReply OnDragDetected(const FGeometry& Geometry, const FPointerEvent& PointerEvent, bool bIsSource);
+
+	FSimpleDelegate& GetOnRequestRename() { return OnRequestRenameDelegate; }
+	FOnSynced& GetOnSynced() { return OnSyncedDelegate; }
+	FOnChildRequestedDeletion& OnChildRequestedDeletion() { return OnChildRequestedDeletionDelegate; }
+
+	TWeakPtr<FHierarchyElementViewModel> GetParent() { return Parent; }
+
+protected:
+	/** Should return true if draggable. An optional message can be provided if false that will show as a slate notification. */
+	virtual FCanPerformActionResults CanDragInternal() { return false; }
+
+	/** Should return true if renamable */
+	virtual bool CanRenameInternal() { return false; }
+
+	virtual void RenameInternal(FName NewName) {}
+	
+	/** Should return true if deletable. By default, we can delete items in the hierarchy, not in the source. */
+	virtual bool CanDeleteInternal() { return IsForHierarchy(); }
+
+	/** Should return true if the given drag drop operation is allowed to succeed. */
+	virtual FCanPerformActionResults CanDropOnInternal(TSharedPtr<FHierarchyElementViewModel>, EItemDropZone ItemDropZone);
+	
+	/** Override this to handle drop-on logic. This is called when an item has been dropped on the item that has implemented this function. */
+	virtual void OnDroppedOnInternal(TSharedPtr<FHierarchyElementViewModel>, EItemDropZone ItemDropZone) { }
+
+	/** Can be overridden to support sections in the source list.
+	 * In the hierarchy only categories can be parented directly to the root, but using this it is possible to add items to custom sections in the source panel.
+	 * This will only work for top-level objects, i.e. anything directly under the root. */
+	virtual const UHierarchySection* GetSectionInternal() const { return nullptr; }
+private:
+	/** Optionally implement this to refresh dependent data. */
+	virtual void RefreshChildrenDataInternal() {}
+	/** Optionally implement this to further customize the view model sync process.
+	 * An example for this is how the root view model handles sections, as sections exist outside the children hierarchy */
+	virtual void SyncViewModelsToDataInternal() {}
+	/** Optionally implement this to handle shutdown logic.
+	 * An example for this is when a section gets deleted, it iterates over all categories to null out the associated section  */
+	virtual void FinalizeInternal() {}
+
+	void PropagateOnChildSynced();
+protected:
+	/** The hierarchy element this view model represents. Assumed valid while this view model exists. */
+	UHierarchyElement* const Element;
+
+	/** Parent should be valid for all instances of this struct except for root objects */
+	TWeakPtr<FHierarchyElementViewModel> Parent;
+	TArray<TSharedPtr<FHierarchyElementViewModel>> Children;
+	
+	TWeakObjectPtr<UDataHierarchyViewModelBase> HierarchyViewModel;
+	
+	TArray<FOnFilterChild> ChildFilters;
+	mutable TArray<TSharedPtr<FHierarchyElementViewModel>> FilteredChildren;
+	
+	FSimpleDelegate OnRequestRenameDelegate;
+	FOnSynced OnSyncedDelegate;
+	FOnChildRequestedDeletion OnChildRequestedDeletionDelegate;
+	
+	bool bRenamePending = false;
+	bool bIsForHierarchy = false;
+};
+
+template <class DataClass, class ViewModelClass>
+void FHierarchyElementViewModel::GetChildrenViewModelsForType(TArray<TSharedPtr<ViewModelClass>>& OutChildren, bool bRecursive)
+{
+	for(auto& Child : Children)
+	{
+		if(Child->GetData()->IsA<DataClass>())
+		{
+			OutChildren.Add(StaticCastSharedPtr<ViewModelClass>(Child));
+		}
+	}
+
+	if(bRecursive)
+	{
+		for(auto& Child : Children)
+		{
+			Child->GetChildrenViewModelsForType<DataClass, ViewModelClass>(OutChildren, bRecursive);
+		}
+	}
+}
+
+struct DATAHIERARCHYEDITOR_API FHierarchyRootViewModel : FHierarchyElementViewModel
+{
+	DECLARE_DELEGATE(FOnSyncPropagated)
+	DECLARE_DELEGATE(FOnSectionsChanged)
+	DECLARE_DELEGATE_OneParam(FOnSingleSectionChanged, TSharedPtr<FHierarchySectionViewModel> AddedSection)
+	
+	FHierarchyRootViewModel(UHierarchyElement* InItem, TWeakObjectPtr<UDataHierarchyViewModelBase> InHierarchyViewModel, bool bInIsForHierarchy) : FHierarchyElementViewModel(InItem, nullptr, InHierarchyViewModel, bInIsForHierarchy) {}
+	virtual ~FHierarchyRootViewModel() override;
+
+	virtual void Initialize() override;
+	
+	virtual FCanPerformActionResults CanDropOnInternal(TSharedPtr<FHierarchyElementViewModel> DraggedElement, EItemDropZone ItemDropZone) override;
+	virtual void OnDroppedOnInternal(TSharedPtr<FHierarchyElementViewModel> DroppedItem, EItemDropZone ItemDropZone) override;
+
+	virtual bool CanHaveChildren() const override { return true; }
+
+	TSharedPtr<struct FHierarchySectionViewModel> AddSection();
+	void DeleteSection(TSharedPtr<FHierarchyElementViewModel> SectionViewModel);
+	TArray<TSharedPtr<struct FHierarchySectionViewModel>>& GetSectionViewModels() { return SectionViewModels; }
+
+	FOnSyncPropagated& OnSyncPropagated() { return OnSyncPropagatedDelegate; }
+
+	/** General purpose delegate for when sections change */
+	FOnSectionsChanged& OnSectionsChanged() { return OnSectionsChangedDelegate; }
+	/** Delegates for when a section is added or removed */
+	FOnSingleSectionChanged& OnSectionAdded() { return OnSectionAddedDelegate; }
+	FOnSingleSectionChanged& OnSectionDeleted() { return OnSectionDeletedDelegate; }
+
+private:
+	void PropagateOnSynced();
+	virtual void SyncViewModelsToDataInternal() override;
+	TArray<TSharedPtr<struct FHierarchySectionViewModel>> SectionViewModels;
+
+	FOnSyncPropagated OnSyncPropagatedDelegate;
+	FOnSingleSectionChanged OnSectionAddedDelegate;
+	FOnSingleSectionChanged OnSectionDeletedDelegate;
+	FOnSectionsChanged OnSectionsChangedDelegate;
+};
+
+struct DATAHIERARCHYEDITOR_API FHierarchySectionViewModel : FHierarchyElementViewModel
+{
+	FHierarchySectionViewModel(UHierarchySection* InItem, TSharedRef<FHierarchyRootViewModel> InParent, TWeakObjectPtr<UDataHierarchyViewModelBase> InHierarchyViewModel) : FHierarchyElementViewModel(InItem, InParent, InHierarchyViewModel, InParent->IsForHierarchy()) {}
+	
+	virtual ~FHierarchySectionViewModel() override {}
+
+	void SetSectionName(FName InSectionName) const { Cast<UHierarchySection>(Element)->SetSectionName(InSectionName); }
+	FName GetSectionName() const { return Cast<UHierarchySection>(Element)->GetSectionName(); }
+	
+	void SetSectionNameAsText(const FText& Text) const { Cast<UHierarchySection>(Element)->SetSectionNameAsText(Text); }
+	FText GetSectionNameAsText() const { return Cast<UHierarchySection>(Element)->GetSectionNameAsText(); }
+	FText GetSectionTooltip() const { return Cast<UHierarchySection>(Element)->GetTooltip(); }
+	
+	void SetSectionImage(const FSlateBrush* InSectionImage) { SectionImage = InSectionImage; }
+	const FSlateBrush* GetSectionImage() const { return SectionImage; }
+
+protected:
+	/** Only hierarchy sections are editable */
+	virtual FCanPerformActionResults IsEditableByUser() override { return FCanPerformActionResults(IsForHierarchy()); }
+	virtual bool CanHaveChildren() const override { return false; }
+	virtual FCanPerformActionResults CanDragInternal() override { return bIsForHierarchy; }
+	/** We can only rename hierarchy sections */
+	virtual bool CanRenameInternal() override { return IsForHierarchy(); }
+	virtual void RenameInternal(FName NewName) override { GetDataMutable<UHierarchySection>()->SetSectionName(NewName); }
+
+	virtual FCanPerformActionResults CanDropOnInternal(TSharedPtr<FHierarchyElementViewModel> DraggedItem, EItemDropZone ItemDropZone) override;
+	virtual void OnDroppedOnInternal(TSharedPtr<FHierarchyElementViewModel> DroppedItem, EItemDropZone ItemDropZone) override;
+
+	virtual void FinalizeInternal() override;
+
+private:
+	const FSlateBrush* SectionImage = nullptr;
+};
+
+struct DATAHIERARCHYEDITOR_API FHierarchyItemViewModel : FHierarchyElementViewModel
+{
+	FHierarchyItemViewModel(UHierarchyItem* InElement, TSharedRef<FHierarchyElementViewModel> InParent, TWeakObjectPtr<UDataHierarchyViewModelBase> InHierarchyViewModel) : FHierarchyElementViewModel(InElement, InParent, InHierarchyViewModel, InParent->IsForHierarchy()) {}
+	
+	virtual ~FHierarchyItemViewModel() override {}
+
+	virtual FCanPerformActionResults IsEditableByUser() override { return FCanPerformActionResults(true); }
+	virtual bool CanHaveChildren() const override { return false; }
+	virtual FCanPerformActionResults CanDragInternal() override { return true; }
+
+	virtual FCanPerformActionResults CanDropOnInternal(TSharedPtr<FHierarchyElementViewModel> DraggedItem, EItemDropZone ItemDropZone) override;
+	virtual void OnDroppedOnInternal(TSharedPtr<FHierarchyElementViewModel> DroppedItem, EItemDropZone ItemDropZone) override;
+};
+
+struct DATAHIERARCHYEDITOR_API FHierarchyCategoryViewModel : FHierarchyElementViewModel
+{
+	FHierarchyCategoryViewModel(UHierarchyCategory* InCategory, TSharedRef<FHierarchyElementViewModel> InParent, TWeakObjectPtr<UDataHierarchyViewModelBase> InHierarchyViewModel) : FHierarchyElementViewModel(InCategory, InParent, InHierarchyViewModel, InParent->IsForHierarchy()) {}
+	virtual ~FHierarchyCategoryViewModel() override{}
+
+	FText GetCategoryName() const { return GetData<UHierarchyCategory>()->GetCategoryAsText(); }
+	
+	virtual FCanPerformActionResults IsEditableByUser() override { return FCanPerformActionResults(true); }
+	virtual bool CanHaveChildren() const override { return true; }
+	virtual FCanPerformActionResults CanDragInternal() override { return true; }
+	virtual bool CanRenameInternal() override { return true; }
+	virtual void RenameInternal(FName NewName) override { GetDataMutable<UHierarchyCategory>()->SetCategoryName(NewName); }
+	virtual const UHierarchySection* GetSectionInternal() const override { return Cast<UHierarchyCategory>(Element)->GetSection(); }
+
+	bool IsTopCategoryActive() const;
+
+	virtual FCanPerformActionResults CanDropOnInternal(TSharedPtr<FHierarchyElementViewModel>, EItemDropZone ItemDropZone) override;
+	virtual void OnDroppedOnInternal(TSharedPtr<FHierarchyElementViewModel> DroppedItem, EItemDropZone ItemDropZone) override;
+};
+
+class DATAHIERARCHYEDITOR_API FSectionDragDropOp : public FHierarchyDragDropOp
+{
+public:
+	DRAG_DROP_OPERATOR_TYPE(FSectionDragDropOp, FHierarchyDragDropOp)
+
+	FSectionDragDropOp(TSharedPtr<FHierarchySectionViewModel> SectionViewModel) : FHierarchyDragDropOp(SectionViewModel) {}
+	
+	TWeakPtr<FHierarchySectionViewModel> GetDraggedSection() const { return StaticCastSharedPtr<FHierarchySectionViewModel>(DraggedElement.Pin()); }
+private:
+	virtual TSharedRef<SWidget> CreateCustomDecorator() const override;
+};
diff --git a/Engine/Source/Editor/HierarchyEditor/HierarchyEditor.Build.cs b/Engine/Source/Editor/HierarchyEditor/HierarchyEditor.Build.cs
deleted file mode 100644
diff --git a/Engine/Source/Editor/HierarchyEditor/Private/HierarchyEditorCommands.cpp b/Engine/Source/Editor/HierarchyEditor/Private/HierarchyEditorCommands.cpp
deleted file mode 100644
diff --git a/Engine/Source/Editor/HierarchyEditor/Private/HierarchyEditorMisc.cpp b/Engine/Source/Editor/HierarchyEditor/Private/HierarchyEditorMisc.cpp
deleted file mode 100644
diff --git a/Engine/Source/Editor/HierarchyEditor/Private/HierarchyEditorModule.cpp b/Engine/Source/Editor/HierarchyEditor/Private/HierarchyEditorModule.cpp
deleted file mode 100644
diff --git a/Engine/Source/Editor/HierarchyEditor/Private/HierarchyEditorStyle.cpp b/Engine/Source/Editor/HierarchyEditor/Private/HierarchyEditorStyle.cpp
deleted file mode 100644
diff --git a/Engine/Source/Editor/HierarchyEditor/Private/HierarchyViewModelBase.cpp b/Engine/Source/Editor/HierarchyEditor/Private/HierarchyViewModelBase.cpp
deleted file mode 100644
diff --git a/Engine/Source/Editor/HierarchyEditor/Private/Widgets/SHierarchyEditor.cpp b/Engine/Source/Editor/HierarchyEditor/Private/Widgets/SHierarchyEditor.cpp
deleted file mode 100644
diff --git a/Engine/Source/Editor/HierarchyEditor/Public/HierarchyEditorCommands.h b/Engine/Source/Editor/HierarchyEditor/Public/HierarchyEditorCommands.h
deleted file mode 100644
diff --git a/Engine/Source/Editor/HierarchyEditor/Public/HierarchyEditorMisc.h b/Engine/Source/Editor/HierarchyEditor/Public/HierarchyEditorMisc.h
deleted file mode 100644
diff --git a/Engine/Source/Editor/HierarchyEditor/Public/HierarchyEditorModule.h b/Engine/Source/Editor/HierarchyEditor/Public/HierarchyEditorModule.h
deleted file mode 100644
diff --git a/Engine/Source/Editor/HierarchyEditor/Public/HierarchyEditorStyle.h b/Engine/Source/Editor/HierarchyEditor/Public/HierarchyEditorStyle.h
deleted file mode 100644
diff --git a/Engine/Source/Editor/HierarchyEditor/Public/HierarchyViewModelBase.h b/Engine/Source/Editor/HierarchyEditor/Public/HierarchyViewModelBase.h
deleted file mode 100644
diff --git a/Engine/Source/Editor/HierarchyEditor/Public/Widgets/SHierarchyEditor.h b/Engine/Source/Editor/HierarchyEditor/Public/Widgets/SHierarchyEditor.h
deleted file mode 100644
diff --git a/Engine/Source/Editor/UnrealEd/Private/EditorEngine.cpp b/Engine/Source/Editor/UnrealEd/Private/EditorEngine.cpp
--- a/Engine/Source/Editor/UnrealEd/Private/EditorEngine.cpp
+++ b/Engine/Source/Editor/UnrealEd/Private/EditorEngine.cpp
@@ -1513,7 +1513,7 @@ void UEditorEngine::LoadDefaultEditorModules()
 			TEXT("StructUtilsTestSuite"),
 			TEXT("MassEntityEditor"),
 			TEXT("SVGDistanceField"),
-			TEXT("HierarchyEditor")
+			TEXT("DataHierarchyEditor")
 		};
 
 	FScopedSlowTask ModuleSlowTask((float)UE_ARRAY_COUNT(ModuleNames));
