Multiple Data Hierarchy Editor fixes:
- Reworked sections
-- Now 'All' sections use a section view model without backing data, instead of no view models at all
-- Moved parts of the section drop logic from the widget to the view model as part of this, cleaning up responsibilities
-- Removed dedicated drop target widget for sections
-- Added SListView style drag & drop
-- Increased padding betwen sections
- Added missing transaction when moving sections
- Transaction for renaming categories or sections now only happens when the name is actually changing. This fixes 2 transactions being created instead of 1 when creating a new section or category (as they would immediately get prompted for rename)
- Removed unused style

diff --git a/Engine/Plugins/FX/Niagara/Source/NiagaraEditor/Private/NiagaraEditorStyle.cpp b/Engine/Plugins/FX/Niagara/Source/NiagaraEditor/Private/NiagaraEditorStyle.cpp
--- a/Engine/Plugins/FX/Niagara/Source/NiagaraEditor/Private/NiagaraEditorStyle.cpp
+++ b/Engine/Plugins/FX/Niagara/Source/NiagaraEditor/Private/NiagaraEditorStyle.cpp
@@ -659,8 +659,6 @@ void FNiagaraEditorStyle::InitDropTarget()
 {
 	Set("NiagaraEditor.DropTarget.BackgroundColor", FLinearColor(1.0f, 1.0f, 1.0f, 0.25f));
 	Set("NiagaraEditor.DropTarget.BackgroundColorHover", FLinearColor(1.0f, 1.0f, 1.0f, 0.1f));
-	Set("NiagaraEditor.DropTarget.BorderVertical", new IMAGE_BRUSH("Icons/StackDropTargetBorder_Vertical", FVector2D(2, 8), FLinearColor::White, ESlateBrushTileType::Vertical));
-	Set("NiagaraEditor.DropTarget.BorderHorizontal", new IMAGE_BRUSH("Icons/StackDropTargetBorder_Horizontal", FVector2D(8, 2), FLinearColor::White, ESlateBrushTileType::Horizontal));
 }
 
 void FNiagaraEditorStyle::InitScriptGraph()
diff --git a/Engine/Source/Editor/DataHierarchyEditor/Internal/Widgets/SDataHierarchyEditor.h b/Engine/Source/Editor/DataHierarchyEditor/Internal/Widgets/SDataHierarchyEditor.h
--- a/Engine/Source/Editor/DataHierarchyEditor/Internal/Widgets/SDataHierarchyEditor.h
+++ b/Engine/Source/Editor/DataHierarchyEditor/Internal/Widgets/SDataHierarchyEditor.h
@@ -25,32 +25,40 @@ struct FHierarchySectionViewModel;
 class DATAHIERARCHYEDITOR_API SHierarchySection : public SCompoundWidget
 {
 	DECLARE_DELEGATE_OneParam(FOnSectionActivated, TSharedPtr<FHierarchySectionViewModel> SectionViewModel)
+
+	static const FSlateBrush LeftDropBrush;
+	static const FSlateBrush RightDropBrush;
 	
 	SLATE_BEGIN_ARGS(SHierarchySection)
-		: _bForbidDropOn(false)
 		{}
 		SLATE_ATTRIBUTE(ECheckBoxState, IsSectionActive)
 		SLATE_EVENT(FOnSectionActivated, OnSectionActivated)
-		SLATE_ARGUMENT(bool, bForbidDropOn)
 	SLATE_END_ARGS()
 
 public:
-	void Construct(const FArguments& InArgs, TSharedPtr<struct FHierarchySectionViewModel> InSection, TWeakObjectPtr<UDataHierarchyViewModelBase> HierarchyViewModel);
+	void Construct(const FArguments& InArgs, TSharedPtr<struct FHierarchySectionViewModel> InSection);
 	virtual ~SHierarchySection() override;
 
-	bool OnCanAcceptDrop(TSharedPtr<FDragDropOperation> DragDropOperation, EItemDropZone ItemDropZone) const;
-	FReply OnDroppedOn(const FGeometry&, const FDragDropEvent& DragDropEvent, EItemDropZone DropZone) const;
-
 	void TryEnterEditingMode() const;
 
 	TSharedPtr<struct FHierarchySectionViewModel> GetSectionViewModel();
 private:
+	virtual int32 OnPaint(const FPaintArgs& Args, const FGeometry& AllottedGeometry, const FSlateRect& MyCullingRect, FSlateWindowElementList& OutDrawElements, int32 LayerId, const FWidgetStyle& InWidgetStyle, bool bParentEnabled) const override;
+	int32 PaintDropIndicator(const FPaintArgs& Args, const FGeometry& Geometry, FSlateRect SlateRect, FSlateWindowElementList& OutDrawElements, int32 INT32, const FWidgetStyle& WidgetStyle, bool bParentEnabled) const;
+	int32 OnPaintDropIndicator(EItemDropZone ItemDropZone, const FPaintArgs& Args, const FGeometry& Geometry, FSlateRect SlateRect, FSlateWindowElementList& OutDrawElements, int32 LayerId, const FWidgetStyle& WidgetStyle, bool bParentEnabled) const;
+
 	virtual FReply OnMouseButtonDown(const FGeometry& MyGeometry, const FPointerEvent& MouseEvent) override;
 	virtual FReply OnMouseButtonUp(const FGeometry& MyGeometry, const FPointerEvent& MouseEvent) override;
+
+	virtual FReply OnDrop(const FGeometry& MyGeometry, const FDragDropEvent& DragDropEvent) override;
+	
+	bool OnCanAcceptDrop(TSharedPtr<FDragDropOperation> DragDropOperation, EItemDropZone ItemDropZone) const;
+	FReply OnDroppedOn(const FGeometry&, const FDragDropEvent& DragDropEvent, EItemDropZone DropZone) const;
 	
 	virtual void OnDragEnter(const FGeometry& MyGeometry, const FDragDropEvent& DragDropEvent) override;
 	virtual FReply OnDragDetected(const FGeometry& MyGeometry, const FPointerEvent& MouseEvent) override;
 	virtual void OnDragLeave(const FDragDropEvent& DragDropEvent) override;
+	virtual FReply OnDragOver(const FGeometry& MyGeometry, const FDragDropEvent& DragDropEvent) override;
 
 	virtual FReply OnKeyDown(const FGeometry& MyGeometry, const FKeyEvent& InKeyEvent) override;
 	
@@ -68,17 +76,25 @@ private:
 	ECheckBoxState GetSectionCheckState() const;
 	void OnSectionCheckChanged(ECheckBoxState NewState);
 	EActiveTimerReturnType ActivateSectionIfDragging(double CurrentTime, float DeltaTime) const;
+
+	const FSlateBrush* GetDropIndicatorBrush(EItemDropZone ItemDropZone) const;
+
+	/** @return the zone (above, onto, below) based on where the user is hovering over within the row */
+	EItemDropZone ZoneFromPointerPosition(UE::Slate::FDeprecateVector2DParameter LocalPointerPos, UE::Slate::FDeprecateVector2DParameter LocalSize);
+
 private:
 	TSharedPtr<SMenuAnchor> MenuAnchor;
 	TSharedPtr<SCheckBox> CheckBox;
 	TSharedPtr<SInlineEditableTextBlock> InlineEditableTextBlock;
 	TWeakObjectPtr<UDataHierarchyViewModelBase> HierarchyViewModel;
-	TSharedPtr<FHierarchySectionViewModel> SectionViewModel;
+	TWeakPtr<FHierarchySectionViewModel> SectionViewModelWeak;
 private:
 	TAttribute<ECheckBoxState> IsSectionActive;
 	FOnSectionActivated OnSectionActivatedDelegate;
-	bool bForbidDropOn = false;
+	
 	mutable bool bDraggedOn = false;
+	
+	mutable TOptional<EItemDropZone> CurrentItemDropZone;
 };
 
 class DATAHIERARCHYEDITOR_API SDataHierarchyEditor : public SCompoundWidget, public FNotifyHook
@@ -226,6 +242,7 @@ private:
 	TSharedPtr<FHierarchyRootViewModel> SourceRootViewModel;
 	TSharedPtr<STreeView<TSharedPtr<FHierarchyElementViewModel>>> SourceTreeView;
 	TSharedPtr<STreeView<TSharedPtr<FHierarchyElementViewModel>>> HierarchyTreeView;
+	TSharedPtr<FHierarchySectionViewModel> DefaultSourceSectionViewModel;
 	TWeakPtr<struct FHierarchySectionViewModel> ActiveSourceSection;
 	TSharedPtr<class SWrapBox> SourceSectionBox;
 	TSharedPtr<class SWrapBox> HierarchySectionBox;
@@ -256,26 +273,6 @@ public:
 	void OnRenameCategory(const FText& NewText, ETextCommit::Type) const;
 	
 private:
-	TWeakPtr<struct FHierarchyCategoryViewModel> CategoryViewModel;
+	TWeakPtr<struct FHierarchyCategoryViewModel> CategoryViewModelWeak;
 	TSharedPtr<SInlineEditableTextBlock> InlineEditableTextBlock;
 };
-
-class SSectionDragDropTarget : public SDropTarget
-{
-public:
-	SLATE_BEGIN_ARGS(SSectionDragDropTarget)
-		: _DropTargetArgs(SDropTarget::FArguments())
-	{}	
-		SLATE_ARGUMENT(SDropTarget::FArguments, DropTargetArgs)
-	SLATE_END_ARGS()
-
-	void Construct(const FArguments& InArgs, TSharedPtr<FHierarchySectionViewModel> OwningSection, EItemDropZone ItemDropZone);
-	
-	virtual FReply OnDragOver(const FGeometry& MyGeometry, const FDragDropEvent& DragDropEvent) override;
-	virtual void OnDragLeave(const FDragDropEvent& DragDropEvent) override;
-
-private:
-	TSharedPtr<FHierarchySectionViewModel> OwningSection;
-	/** The drop zone this drop target represents */
-	EItemDropZone DropZone = EItemDropZone::OntoItem;
-};
diff --git a/Engine/Source/Editor/DataHierarchyEditor/Private/DataHierarchyEditorStyle.cpp b/Engine/Source/Editor/DataHierarchyEditor/Private/DataHierarchyEditorStyle.cpp
--- a/Engine/Source/Editor/DataHierarchyEditor/Private/DataHierarchyEditorStyle.cpp
+++ b/Engine/Source/Editor/DataHierarchyEditor/Private/DataHierarchyEditorStyle.cpp
@@ -31,6 +31,16 @@ void FDataHierarchyEditorStyle::Shutdown()
 
 void FDataHierarchyEditorStyle::InitDataHierarchyEditor()
 {
+	SetCoreContentRoot(FPaths::EngineContentDir() / TEXT("Slate"));
+
+	const FSlateColor SelectorColor = FAppStyle::GetSlateColor("SelectorColor");
+	const FSlateColor SelectionColor = FAppStyle::GetSlateColor("SelectionColor");
+	const FSlateColor SelectionColor_Inactive = FAppStyle::GetSlateColor("SelectionColor_Inactive");
+	
+	Set("HierarchyEditor.Drop.Section.Left", new CORE_BOX_BRUSH("Common/DropZoneIndicator_Above", FMargin(10.0f / 16.0f, 10.0f / 16.0f, 0.f, 0.f), SelectionColor));
+	Set("HierarchyEditor.Drop.Section.Onto", new CORE_BOX_BRUSH("Common/DropZoneIndicator_Onto", FMargin(4.0f / 16.0f), SelectionColor));
+	Set("HierarchyEditor.Drop.Section.Below", new CORE_BOX_BRUSH("Common/DropZoneIndicator_Below", FMargin(10.0f / 16.0f, 0.f, 0.f, 10.0f / 16.0f), SelectionColor));
+	
 	const FTextBlockStyle NormalText = FAppStyle::GetWidgetStyle<FTextBlockStyle>("NormalText");
 	FSlateFontInfo CategoryFont = FAppStyle::Get().GetFontStyle(TEXT("DetailsView.CategoryFontStyle"));
 	CategoryFont.Size = 11;
@@ -50,9 +60,6 @@ void FDataHierarchyEditorStyle::InitDataHierarchyEditor()
 	Set("HierarchyEditor.CategoryTextBlock", CategoryText);
 	Set("HierarchyEditor.Category", HierarchyCategoryTextStyle);
 
-	Set("HierarchyEditor.Stack.DropTarget.BorderVertical", new IMAGE_BRUSH("Icons/StackDropTargetBorder_Vertical", FVector2D(2, 8), FLinearColor::White, ESlateBrushTileType::Vertical));
-	Set("HierarchyEditor.Stack.DropTarget.BorderHorizontal", new IMAGE_BRUSH("Icons/StackDropTargetBorder_Horizontal", FVector2D(8, 2), FLinearColor::White, ESlateBrushTileType::Horizontal));
-
 	FButtonStyle SimpleButtonStyle = FAppStyle::Get().GetWidgetStyle<FButtonStyle>("SimpleButton");
 	FButtonStyle ButtonStyle = FButtonStyle(SimpleButtonStyle)
 		.SetNormalForeground(FStyleColors::Foreground)
@@ -63,10 +70,7 @@ void FDataHierarchyEditorStyle::InitDataHierarchyEditor()
 		.SetPressedPadding(FMargin(8.f, 3.f, 8.f, 1.f));
 
 	Set("HierarchyEditor.ButtonStyle", ButtonStyle);
-
-	const FSlateColor SelectorColor = FAppStyle::GetSlateColor("SelectorColor");
-	const FSlateColor SelectionColor = FAppStyle::GetSlateColor("SelectionColor");
-	const FSlateColor SelectionColor_Inactive = FAppStyle::GetSlateColor("SelectionColor_Inactive");
+	
 	FLinearColor HighlightColor = FLinearColor(0.068f, 0.068f, 0.068f);
 	
 	Set("HierarchyEditor.Color.Highlight", HighlightColor);
@@ -83,9 +87,9 @@ void FDataHierarchyEditorStyle::InitDataHierarchyEditor()
 		.SetInactiveHoveredBrush( CORE_IMAGE_BRUSH( "Common/Selection", CoreStyleConstants::Icon8x8, SelectionColor_Inactive))
 		.SetActiveHighlightedBrush( CORE_IMAGE_BRUSH("Common/Selection", CoreStyleConstants::Icon8x8, HighlightColor))
 		.SetInactiveHighlightedBrush( CORE_IMAGE_BRUSH("Common/Selection", CoreStyleConstants::Icon8x8, HighlightColor))
-		.SetDropIndicator_Above(CORE_BOX_BRUSH("Common/DropZoneIndicator_Above", FMargin(10.0f / 16.0f, 10.0f / 16.0f, 0, 0), FStarshipCoreStyle::GetCoreStyle().GetSlateColor("SelectorColor")))
-		.SetDropIndicator_Onto(CORE_BOX_BRUSH("Common/DropZoneIndicator_Onto", FMargin(4.0f / 16.0f), FStarshipCoreStyle::GetCoreStyle().GetSlateColor("SelectorColor")))
-		.SetDropIndicator_Below(CORE_BOX_BRUSH("Common/DropZoneIndicator_Below", FMargin(10.0f / 16.0f, 0, 0, 10.0f / 16.0f), FStarshipCoreStyle::GetCoreStyle().GetSlateColor("SelectorColor")));
+		.SetDropIndicator_Above(CORE_BOX_BRUSH("Common/DropZoneIndicator_Above", FMargin(10.0f / 16.0f, 10.0f / 16.0f, 0, 0), SelectionColor))
+		.SetDropIndicator_Onto(CORE_BOX_BRUSH("Common/DropZoneIndicator_Onto", FMargin(4.0f / 16.0f), SelectionColor))
+		.SetDropIndicator_Below(CORE_BOX_BRUSH("Common/DropZoneIndicator_Below", FMargin(10.0f / 16.0f, 0, 0, 10.0f / 16.0f), SelectionColor));
 	
 	Set("HierarchyEditor.Row.Category", HierarchyEditorCategoryRowStyle);
 	
@@ -94,12 +98,11 @@ void FDataHierarchyEditorStyle::InitDataHierarchyEditor()
 	HierarchyEditorCategoryStyle.TextStyle.SetFontSize(10.f);
 	Set("HierarchyEditor.CategoryTextStyle", HierarchyEditorCategoryStyle);
 	
-	Set("HierarchyEditor.RootDropIcon", new IMAGE_BRUSH_SVG("Starship/Common/caret-down", CoreStyleConstants::Icon20x20));
+	Set("HierarchyEditor.RootDropIcon", new CORE_IMAGE_BRUSH_SVG("Starship/Common/caret-down", CoreStyleConstants::Icon20x20));
 }
 
 FDataHierarchyEditorStyle::FDataHierarchyEditorStyle() : FSlateStyleSet("DataHierarchyEditorStyle")
 {
-	SetContentRoot(FPaths::EngineContentDir() / TEXT("Slate"));
 	InitDataHierarchyEditor();
 }
 
diff --git a/Engine/Source/Editor/DataHierarchyEditor/Private/DataHierarchyViewModelBase.cpp b/Engine/Source/Editor/DataHierarchyEditor/Private/DataHierarchyViewModelBase.cpp
--- a/Engine/Source/Editor/DataHierarchyEditor/Private/DataHierarchyViewModelBase.cpp
+++ b/Engine/Source/Editor/DataHierarchyEditor/Private/DataHierarchyViewModelBase.cpp
@@ -467,8 +467,9 @@ void UDataHierarchyViewModelBase::Initialize()
 	HierarchyRootViewModel->AddChildFilter(FHierarchyElementViewModel::FOnFilterChild::CreateUObject(this, &UDataHierarchyViewModelBase::FilterForHierarchySection));
 	HierarchyRootViewModel->AddChildFilter(FHierarchyElementViewModel::FOnFilterChild::CreateUObject(this, &UDataHierarchyViewModelBase::FilterForUncategorizedRootItemsInAllSection));
 	HierarchyRootViewModel->SyncViewModelsToData();
-	
-	SetActiveHierarchySection(nullptr);
+
+	DefaultHierarchySectionViewModel = MakeShared<FHierarchySectionViewModel>(nullptr, GetHierarchyRootViewModel().ToSharedRef(), this);
+	SetActiveHierarchySection(DefaultHierarchySectionViewModel);
 	
 	InitializeInternal();
 
@@ -532,7 +533,6 @@ TSubclassOf<UHierarchySection> UDataHierarchyViewModelBase::GetSectionDataClass(
 	return UHierarchySection::StaticClass();
 }
 
-
 void UDataHierarchyViewModelBase::ForceFullRefresh()
 {
 	RefreshSourceItemsRequestedDelegate.ExecuteIfBound();
@@ -606,7 +606,14 @@ bool UDataHierarchyViewModelBase::FilterForHierarchySection(TSharedPtr<const FHi
 {
 	if(ActiveHierarchySection.IsValid())
 	{
-		return GetActiveHierarchySectionData() == ViewModel->GetSection();
+		// If the currently selected section data is nullptr, it's the All section, and we let everything pass
+		if(ActiveHierarchySection.Pin()->GetData() == nullptr)
+		{
+			return true;
+		}
+
+		// if not, we check against identical section data
+		return ActiveHierarchySection.Pin()->GetData() == ViewModel->GetSection();
 	}
 
 	return true;
@@ -614,9 +621,14 @@ bool UDataHierarchyViewModelBase::FilterForHierarchySection(TSharedPtr<const FHi
 
 bool UDataHierarchyViewModelBase::FilterForUncategorizedRootItemsInAllSection(TSharedPtr<const FHierarchyElementViewModel> ViewModel) const
 {
-	// we want to filter out all items that are directly added to the root if we aren't in the 'All' section
 	if(ActiveHierarchySection.IsValid())
 	{
+		// we want to filter out all items that are directly added to the root if we aren't in the 'All' section
+		if(ActiveHierarchySection.Pin()->GetData() == nullptr)
+		{
+			return true;
+		}
+		
 		return ViewModel->GetData<UHierarchyCategory>() != nullptr;
 	}
 
@@ -731,17 +743,17 @@ void UDataHierarchyViewModelBase::SetActiveHierarchySection(TSharedPtr<FHierarch
 
 TSharedPtr<FHierarchySectionViewModel> UDataHierarchyViewModelBase::GetActiveHierarchySection() const
 {
-	return ActiveHierarchySection.IsValid() ? ActiveHierarchySection.Pin() : nullptr;
+	return ActiveHierarchySection.Pin();
 }
 
 UHierarchySection* UDataHierarchyViewModelBase::GetActiveHierarchySectionData() const
 {
-	return ActiveHierarchySection.IsValid() ? ActiveHierarchySection.Pin()->GetDataMutable<UHierarchySection>() : nullptr;
+	return ActiveHierarchySection.Pin()->GetDataMutable<UHierarchySection>();
 }
 
 bool UDataHierarchyViewModelBase::IsHierarchySectionActive(const UHierarchySection* Section) const
 {
-	return ActiveHierarchySection == nullptr || ActiveHierarchySection.Pin()->GetData() == Section;
+	return ActiveHierarchySection.Pin()->GetData() == Section;
 }
 
 FString UDataHierarchyViewModelBase::OnElementToStringDebug(TSharedPtr<FHierarchyElementViewModel> ElementViewModel) const
@@ -1330,8 +1342,74 @@ void FHierarchyRootViewModel::SyncViewModelsToDataInternal()
 		});
 }
 
+FString FHierarchySectionViewModel::ToString() const
+{
+	return GetSectionNameAsText().ToString();
+}
+
+void FHierarchySectionViewModel::SetSectionName(FName InSectionName)
+{
+	Cast<UHierarchySection>(Element)->SetSectionName(InSectionName);
+}
+
+FName FHierarchySectionViewModel::GetSectionName() const
+{
+	if(UHierarchySection* Section = Cast<UHierarchySection>(Element))
+	{
+		return Section->GetSectionName();
+	}
+
+	return NAME_None;
+}
+
+void FHierarchySectionViewModel::SetSectionNameAsText(const FText& Text)
+{
+	Cast<UHierarchySection>(Element)->SetSectionNameAsText(Text);
+}
+
+FText FHierarchySectionViewModel::GetSectionNameAsText() const
+{
+	if(UHierarchySection* Section = Cast<UHierarchySection>(Element))
+	{
+		return Section->GetSectionNameAsText();
+	}
+
+	return LOCTEXT("DefaultSectionName", "All");
+}
+
+FText FHierarchySectionViewModel::GetSectionTooltip() const
+{
+	if(UHierarchySection* Section = Cast<UHierarchySection>(Element))
+	{
+		return Section->GetTooltip();
+	}
+	
+	return FText::GetEmpty();
+}
+
+FHierarchyElementViewModel::FCanPerformActionResults FHierarchySectionViewModel::CanDragInternal()
+{
+	// We only allow hierarchy sections to be dragged, excluding the All section that has no valid data
+	return IsForHierarchy() && GetData() != nullptr;
+}
+
+bool FHierarchySectionViewModel::CanRenameInternal()
+{
+	return IsForHierarchy() && GetData() != nullptr;
+}
+
+bool FHierarchySectionViewModel::CanDeleteInternal()
+{
+	return IsForHierarchy() && GetData() != nullptr;
+}
+
 FHierarchyElementViewModel::FCanPerformActionResults FHierarchySectionViewModel::CanDropOnInternal(TSharedPtr<FHierarchyElementViewModel> DraggedElement, EItemDropZone ItemDropZone)
 {
+	if(bDropDisallowed)
+	{
+		return false;
+	}
+	
 	FCanPerformActionResults Results(false);
 	// we don't allow dropping onto source sections and we don't specify a message as the sections aren't going to light up as valid drop targets
 	if(IsForHierarchy() == false)
@@ -1346,40 +1424,40 @@ FHierarchyElementViewModel::FCanPerformActionResults FHierarchySectionViewModel:
 			FText Message = LOCTEXT("DropCategoryOnSectionDragMessage", "Add {0} to section {1}");
 			Message = FText::FormatOrdered(Message, FText::FromString(DraggedElement->ToString()), FText::FromString(ToString()));
 			
-			Results.CanPerformMessage = Message;
 			Results.bCanPerform = GetData() != Category->GetSection();
+			Results.CanPerformMessage = Results.bCanPerform ? Message : FText::GetEmpty();
 		}
 	}
-	else if(UHierarchySection* Section = Cast<UHierarchySection>(DraggedElement->GetDataMutable()))
+	else if(UHierarchySection* DraggedSection = Cast<UHierarchySection>(DraggedElement->GetDataMutable()))
 	{
-		const bool bSameSection = GetData() == Section;
+		const bool bSameSection = GetData() == DraggedSection;
 
-		if(ItemDropZone == EItemDropZone::OntoItem && bSameSection == false)
+		// If we drag a section onto a section, nothing happens
+		if(ItemDropZone == EItemDropZone::OntoItem)
+		{
+			Results.bCanPerform = false;
+			return Results;
+		}
+
+		// The 'All' section does not accept any drop actions.
+		if(GetData() == nullptr)
 		{
-			FText Message = LOCTEXT("CantDropSectionOnSectionDragMessage", "Can't drop section on a section.");
-			Results.CanPerformMessage = Message;
 			Results.bCanPerform = false;
 			return Results;
 		}
 		
-		int32 DraggedElementIndex = GetHierarchyViewModel()->GetHierarchyRoot()->GetSectionIndex(Section);
+		int32 DraggedSectionIndex = GetHierarchyViewModel()->GetHierarchyRoot()->GetSectionIndex(DraggedSection);
 		int32 InsertionIndex = GetHierarchyViewModel()->GetHierarchyRoot()->GetSectionIndex(GetDataMutable<UHierarchySection>());
 		// we add 1 to the insertion index if it's below an item because we either want to insert at the current index to place the item above, or at current+1 for below
 		InsertionIndex += ItemDropZone == EItemDropZone::AboveItem ? -1 : 1;
 
-		Results.bCanPerform = !bSameSection && DraggedElementIndex != InsertionIndex;
+		Results.bCanPerform = !bSameSection && DraggedSectionIndex != InsertionIndex;
 
 		if(Results.bCanPerform)
 		{
-			if(ItemDropZone == EItemDropZone::AboveItem)
-			{
-				FText Message = LOCTEXT("MoveSectionLeftDragMessage", "Move section {0} to the left");
-				Message = FText::FormatOrdered(Message, FText::FromString(DraggedElement->ToString()));
-				Results.CanPerformMessage = Message;
-			}
-			else if(ItemDropZone == EItemDropZone::BelowItem)
+			if(ItemDropZone != EItemDropZone::OntoItem)
 			{
-				FText Message = LOCTEXT("MoveSectionRightDragMessage", "Move section {0} to the right");
+				FText Message = LOCTEXT("MoveSectionLeftDragMessage", "Move section here");
 				Message = FText::FormatOrdered(Message, FText::FromString(DraggedElement->ToString()));
 				Results.CanPerformMessage = Message;
 			}
@@ -1400,6 +1478,9 @@ void FHierarchySectionViewModel::OnDroppedOnInternal(TSharedPtr<FHierarchyElemen
 {
 	if(DroppedItem->GetData()->IsA<UHierarchySection>())
 	{
+		FScopedTransaction Transaction(LOCTEXT("Transaction_OnSectionMoved", "Moved section"));
+		HierarchyViewModel->GetHierarchyRoot()->Modify();
+
 		UHierarchySection* DraggedSectionData = DroppedItem->GetDataMutable<UHierarchySection>();
 
 		int32 IndexOfThis = HierarchyViewModel->GetHierarchyRoot()->GetSectionData().Find(GetDataMutable<UHierarchySection>());
@@ -1440,32 +1521,29 @@ void FHierarchySectionViewModel::OnDroppedOnInternal(TSharedPtr<FHierarchyElemen
 			HierarchyViewModel->OnHierarchyChanged().Broadcast();
 		}
 	}
-	else
-	{		
-		if(UHierarchyCategory* HierarchyCategory = DroppedItem->GetDataMutable<UHierarchyCategory>())
-		{
-			FScopedTransaction Transaction(LOCTEXT("Transaction_OnSectionDrop", "Moved category to section"));
-			HierarchyViewModel->GetHierarchyRoot()->Modify();
-			
-			HierarchyCategory->SetSection(GetDataMutable<UHierarchySection>());
+	else if(UHierarchyCategory* HierarchyCategory = DroppedItem->GetDataMutable<UHierarchyCategory>())
+	{
+		FScopedTransaction Transaction(LOCTEXT("Transaction_OnSectionDrop", "Moved category to section"));
+		HierarchyViewModel->GetHierarchyRoot()->Modify();
+		
+		HierarchyCategory->SetSection(GetDataMutable<UHierarchySection>());
 
-			// we null out any sections for all contained categories
-			TArray<UHierarchyCategory*> AllChildCategories;
-			HierarchyCategory->GetChildrenOfType<UHierarchyCategory>(AllChildCategories, true);
-			for(UHierarchyCategory* ChildCategory : AllChildCategories)
-			{
-				ChildCategory->SetSection(nullptr);
-			}
+		// we null out any sections for all contained categories
+		TArray<UHierarchyCategory*> AllChildCategories;
+		HierarchyCategory->GetChildrenOfType<UHierarchyCategory>(AllChildCategories, true);
+		for(UHierarchyCategory* ChildCategory : AllChildCategories)
+		{
+			ChildCategory->SetSection(nullptr);
+		}
 
-			// we only need to reparent if the parent isn't already the root. This stops unnecessary reordering
-			if(DroppedItem->GetParent() != HierarchyViewModel->GetHierarchyRootViewModel())
-			{
-				HierarchyViewModel->GetHierarchyRootViewModel()->ReparentToThis(DroppedItem);
-			}
-			
-			HierarchyViewModel->RefreshHierarchyView();
-			HierarchyViewModel->OnHierarchyChanged().Broadcast();
+		// we only need to reparent if the parent isn't already the root. This stops unnecessary reordering
+		if(DroppedItem->GetParent() != HierarchyViewModel->GetHierarchyRootViewModel())
+		{
+			HierarchyViewModel->GetHierarchyRootViewModel()->ReparentToThis(DroppedItem);
 		}
+		
+		HierarchyViewModel->RefreshHierarchyView();
+		HierarchyViewModel->OnHierarchyChanged().Broadcast();
 	}
 }
 
@@ -1473,7 +1551,7 @@ void FHierarchySectionViewModel::FinalizeInternal()
 {
 	if(HierarchyViewModel->GetActiveHierarchySection() == AsShared())
 	{
-		HierarchyViewModel->SetActiveHierarchySection(nullptr);
+		HierarchyViewModel->SetActiveHierarchySection(HierarchyViewModel->GetDefaultHierarchySectionViewModel());
 	}
 
 	// we make sure to reset all categories' section entry that were referencing this section
diff --git a/Engine/Source/Editor/DataHierarchyEditor/Private/Widgets/SDataHierarchyEditor.cpp b/Engine/Source/Editor/DataHierarchyEditor/Private/Widgets/SDataHierarchyEditor.cpp
--- a/Engine/Source/Editor/DataHierarchyEditor/Private/Widgets/SDataHierarchyEditor.cpp
+++ b/Engine/Source/Editor/DataHierarchyEditor/Private/Widgets/SDataHierarchyEditor.cpp
@@ -42,7 +42,7 @@ TSharedRef<SWidget> SummonContextMenu(TArray<TSharedPtr<FHierarchyElementViewMod
 
 void SHierarchyCategory::Construct(const FArguments& InArgs,	TSharedPtr<FHierarchyCategoryViewModel> InCategoryViewModel)
 {
-	CategoryViewModel = InCategoryViewModel;
+	CategoryViewModelWeak = InCategoryViewModel;
 	
 	InCategoryViewModel->GetOnRequestRename().BindSP(this, &SHierarchyCategory::EnterEditingMode);
 
@@ -65,7 +65,7 @@ void SHierarchyCategory::Construct(const FArguments& InArgs,	TSharedPtr<FHierarc
 
 void SHierarchyCategory::EnterEditingMode() const
 {
-	if(CategoryViewModel.Pin()->CanRename())
+	if(CategoryViewModelWeak.Pin()->CanRename())
 	{
 		InlineEditableTextBlock->EnterEditingMode();
 	}
@@ -74,7 +74,7 @@ void SHierarchyCategory::EnterEditingMode() const
 bool SHierarchyCategory::OnVerifyCategoryRename(const FText& NewName, FText& OutTooltip) const
 {
 	TArray<TSharedPtr<FHierarchyCategoryViewModel>> SiblingCategoryViewModels;
-	CategoryViewModel.Pin()->GetParent().Pin()->GetChildrenViewModelsForType<UHierarchyCategory, FHierarchyCategoryViewModel>(SiblingCategoryViewModels);
+	CategoryViewModelWeak.Pin()->GetParent().Pin()->GetChildrenViewModelsForType<UHierarchyCategory, FHierarchyCategoryViewModel>(SiblingCategoryViewModels);
 
 	if(GetCategoryText().ToString() != NewName.ToString())
 	{
@@ -96,133 +96,74 @@ bool SHierarchyCategory::OnVerifyCategoryRename(const FText& NewName, FText& Out
 
 FText SHierarchyCategory::GetCategoryText() const
 {
-	return FText::FromString(CategoryViewModel.Pin()->ToString());
+	return FText::FromString(CategoryViewModelWeak.Pin()->ToString());
 }
 
 void SHierarchyCategory::OnRenameCategory(const FText& NewText, ETextCommit::Type) const
 {
-	FScopedTransaction Transaction(LOCTEXT("Transaction_Rename_Category", "Renamed hierarchy category"));
-	CategoryViewModel.Pin()->GetHierarchyViewModel()->GetHierarchyRoot()->Modify();
-	
-	CategoryViewModel.Pin()->Rename(FName(NewText.ToString()));
+	if(CategoryViewModelWeak.IsValid() && CategoryViewModelWeak.Pin()->GetCategoryName().EqualTo(NewText) == false)
+	{
+		FScopedTransaction Transaction(LOCTEXT("Transaction_Rename_Category", "Renamed hierarchy category"));
+		CategoryViewModelWeak.Pin()->GetHierarchyViewModel()->GetHierarchyRoot()->Modify();
+		
+		CategoryViewModelWeak.Pin()->Rename(FName(NewText.ToString()));
+	}
 }
 
-void SHierarchySection::Construct(const FArguments& InArgs, TSharedPtr<FHierarchySectionViewModel> InSection, TWeakObjectPtr<UDataHierarchyViewModelBase> InHierarchyViewModel)
+void SHierarchySection::Construct(const FArguments& InArgs, TSharedPtr<FHierarchySectionViewModel> InSection)
 {
-	SectionViewModel = InSection;
-	HierarchyViewModel = InHierarchyViewModel;
+	SectionViewModelWeak = InSection;
+	HierarchyViewModel = InSection->GetHierarchyViewModel();
 
 	IsSectionActive = InArgs._IsSectionActive;
 	OnSectionActivatedDelegate = InArgs._OnSectionActivated;
-	bForbidDropOn = InArgs._bForbidDropOn;
 	
-	if(SectionViewModel != nullptr)
-	{
-		SectionViewModel->GetOnRequestRename().BindSP(this, &SHierarchySection::TryEnterEditingMode);
+	InSection->GetOnRequestRename().BindSP(this, &SHierarchySection::TryEnterEditingMode);
 
-		SDropTarget::FArguments LeftDropTargetArgs;
-		SDropTarget::FArguments OntoDropTargetArgs;
-		SDropTarget::FArguments RightDropTargetArgs;
+	SetToolTipText(TAttribute<FText>::CreateSP(this, &SHierarchySection::GetTooltipText));
 
-		LeftDropTargetArgs
-			.OnAllowDrop(this, &SHierarchySection::OnCanAcceptDrop, EItemDropZone::AboveItem)
-			.OnDropped(this, &SHierarchySection::OnDroppedOn, EItemDropZone::AboveItem)
-			.VerticalImage(FAppStyle::GetNoBrush())
-			.HorizontalImage(FDataHierarchyEditorStyle::Get().GetBrush("HierarchyEditor.Stack.DropTarget.BorderHorizontal"));
-
-		OntoDropTargetArgs
-			.OnAllowDrop(this, &SHierarchySection::OnCanAcceptDrop, EItemDropZone::OntoItem)
-			.OnDropped(this, &SHierarchySection::OnDroppedOn, EItemDropZone::OntoItem);
-		
-		RightDropTargetArgs
-			.OnAllowDrop(this, &SHierarchySection::OnCanAcceptDrop, EItemDropZone::BelowItem)
-			.OnDropped(this, &SHierarchySection::OnDroppedOn, EItemDropZone::BelowItem)
-			.VerticalImage(FAppStyle::GetNoBrush())
-			.HorizontalImage(FDataHierarchyEditorStyle::Get().GetBrush("HierarchyEditor.Stack.DropTarget.BorderHorizontal"));
-
-		SetToolTipText(TAttribute<FText>::CreateSP(this, &SHierarchySection::GetTooltipText));
-
-		ChildSlot
+	ChildSlot
+	[
+		SNew(SHorizontalBox)
+		+ SHorizontalBox::Slot()
 		[
-			SNew(SHorizontalBox)
-			+ SHorizontalBox::Slot()
-			.AutoWidth()
-			[
-				SNew(SSectionDragDropTarget, SectionViewModel, EItemDropZone::AboveItem)
-				.DropTargetArgs(LeftDropTargetArgs)
-			]
-			+ SHorizontalBox::Slot()
+			SAssignNew(CheckBox, SCheckBox)
+			.Visibility(EVisibility::HitTestInvisible)
+			.Style(FAppStyle::Get(), "DetailsView.SectionButton")
+			.OnCheckStateChanged(this, &SHierarchySection::OnSectionCheckChanged)
+			.IsChecked(this, &SHierarchySection::GetSectionCheckState)
+			.Padding(FMargin(8.f, 4.f))
 			[
-				SNew(SSectionDragDropTarget, SectionViewModel, EItemDropZone::OntoItem)
-				.DropTargetArgs(OntoDropTargetArgs
-				.Content()
+				SNew(SHorizontalBox)
+				+ SHorizontalBox::Slot()
+				.AutoWidth()
+				.Padding(2.f)
 				[
-					SAssignNew(CheckBox, SCheckBox)
-					.Visibility(EVisibility::HitTestInvisible)
-					.Style(FAppStyle::Get(), "DetailsView.SectionButton")
-					.OnCheckStateChanged(this, &SHierarchySection::OnSectionCheckChanged)
-					.IsChecked(this, &SHierarchySection::GetSectionCheckState)
-					.Padding(FMargin(8.f, 4.f))
-					[
-						SNew(SHorizontalBox)
-						+ SHorizontalBox::Slot()
-						.AutoWidth()
-						.Padding(2.f)
-						[
-							SNew(SImage).Image(SectionViewModel->GetSectionImage() != nullptr ? SectionViewModel->GetSectionImage() : FStyleDefaults::GetNoBrush())
-						]
-						+ SHorizontalBox::Slot()
-						[
-							SAssignNew(InlineEditableTextBlock, SInlineEditableTextBlock)
-							.Visibility(EVisibility::HitTestInvisible)
-							.Text(this, &SHierarchySection::GetText)
-							.OnTextCommitted(this, &SHierarchySection::OnRenameSection)
-							.OnVerifyTextChanged(this, &SHierarchySection::OnVerifySectionRename)
-							.IsSelected(this, &SHierarchySection::IsSectionSelected)
-							.IsReadOnly(this, &SHierarchySection::IsSectionReadOnly)
-						]
-					]
-				])
-			]
-			+ SHorizontalBox::Slot()
-			.AutoWidth()
-			[
-				SNew(SSectionDragDropTarget, SectionViewModel, EItemDropZone::BelowItem)
-				.DropTargetArgs(RightDropTargetArgs)				
-			]
-		];		
-	}
-	// if this section doesn't represent data, it's the "All" widget
-	else
-	{
-		ChildSlot
-		[
-			SNew(SDropTarget)
-			.OnAllowDrop(this, &SHierarchySection::OnCanAcceptDrop, EItemDropZone::OntoItem)
-			.OnDropped(this, &SHierarchySection::OnDroppedOn, EItemDropZone::OntoItem)
-			[
-				SNew(SCheckBox)
-				.Style(FAppStyle::Get(), "DetailsView.SectionButton")
-				.OnCheckStateChanged(this, &SHierarchySection::OnSectionCheckChanged)
-				.IsChecked(this, &SHierarchySection::GetSectionCheckState)
-				.Padding(FMargin(8.f, 4.f))
+					SNew(SImage).Image(InSection->GetSectionImage() != nullptr ? InSection->GetSectionImage() : FStyleDefaults::GetNoBrush())
+				]
+				+ SHorizontalBox::Slot()
 				[
-					SNew(STextBlock)
-					.Text(LOCTEXT("AllSection", "All"))
+					SAssignNew(InlineEditableTextBlock, SInlineEditableTextBlock)
+					.Visibility(EVisibility::HitTestInvisible)
+					.Text(this, &SHierarchySection::GetText)
+					.OnTextCommitted(this, &SHierarchySection::OnRenameSection)
+					.OnVerifyTextChanged(this, &SHierarchySection::OnVerifySectionRename)
+					.IsSelected(this, &SHierarchySection::IsSectionSelected)
+					.IsReadOnly(this, &SHierarchySection::IsSectionReadOnly)
 				]
-			]			
-		];
-	}
+			]
+		]
+	];		
 }
 
 SHierarchySection::~SHierarchySection()
 {
-	SectionViewModel.Reset();
+	SectionViewModelWeak.Reset();
 }
 
 void SHierarchySection::TryEnterEditingMode() const
 {
-	if(SectionViewModel.IsValid() && SectionViewModel->CanRename())
+	if(SectionViewModelWeak.IsValid() && SectionViewModelWeak.Pin()->CanRename())
 	{
 		InlineEditableTextBlock->EnterEditingMode();
 	}
@@ -230,34 +171,59 @@ void SHierarchySection::TryEnterEditingMode() const
 
 TSharedPtr<FHierarchySectionViewModel> SHierarchySection::GetSectionViewModel()
 {
-	return SectionViewModel;
+	return SectionViewModelWeak.Pin();
+}
+
+int32 SHierarchySection::OnPaint(const FPaintArgs& Args, const FGeometry& AllottedGeometry, const FSlateRect& MyCullingRect, FSlateWindowElementList& OutDrawElements, int32 LayerId, const FWidgetStyle& InWidgetStyle, bool bParentEnabled) const
+{
+	LayerId =  SCompoundWidget::OnPaint(Args, AllottedGeometry, MyCullingRect, OutDrawElements, LayerId, InWidgetStyle, bParentEnabled);
+	LayerId = PaintDropIndicator(Args, AllottedGeometry, MyCullingRect, OutDrawElements, LayerId, InWidgetStyle, bParentEnabled);
+
+	return LayerId;
 }
 
-bool SHierarchySection::OnCanAcceptDrop(TSharedPtr<FDragDropOperation> DragDropOperation, EItemDropZone ItemDropZone) const
+int32 SHierarchySection::PaintDropIndicator(const FPaintArgs& Args, const FGeometry& Geometry, FSlateRect SlateRect, FSlateWindowElementList& OutDrawElements, int32 LayerId, const FWidgetStyle& WidgetStyle, bool bParentEnabled) const
 {
-	if(bForbidDropOn)
+	if(CurrentItemDropZone.IsSet())
 	{
-		return false;
+		return OnPaintDropIndicator(CurrentItemDropZone.GetValue(), Args, Geometry, SlateRect, OutDrawElements, LayerId, WidgetStyle, bParentEnabled);
 	}
 
-	if(DragDropOperation->IsOfType<FHierarchyDragDropOp>())
+	return LayerId;
+}
+
+int32 SHierarchySection::OnPaintDropIndicator(EItemDropZone InItemDropZone, const FPaintArgs& Args, const FGeometry& Geometry, FSlateRect SlateRect, FSlateWindowElementList& OutDrawElements, int32 LayerId, const FWidgetStyle& WidgetStyle, bool bParentEnabled) const
+{
+	const FSlateBrush* DropIndicatorBrush = GetDropIndicatorBrush(InItemDropZone);
+
+	const FVector2f LocalSize(Geometry.GetLocalSize());
+	const FVector2f Pivot(LocalSize * 0.5f);
+	const FVector2f RotatedLocalSize(LocalSize.Y, LocalSize.X);
+	FSlateLayoutTransform RotatedTransform(Pivot - RotatedLocalSize * 0.5f);	// Make the box centered to the alloted geometry, so that it can be rotated around the center.
+	
+	FSlateDrawElement::MakeRotatedBox(
+		OutDrawElements,
+		LayerId++,
+		Geometry.ToPaintGeometry(RotatedLocalSize, RotatedTransform),
+		DropIndicatorBrush,
+		ESlateDrawEffect::None,
+		-UE_HALF_PI,	// 90 deg CCW
+		RotatedLocalSize * 0.5f,	// Relative center to the flipped
+		FSlateDrawElement::RelativeToElement,
+		DropIndicatorBrush->GetTint(WidgetStyle) * WidgetStyle.GetColorAndOpacityTint()
+	);
+
+	return LayerId;
+}
+
+bool SHierarchySection::OnCanAcceptDrop(TSharedPtr<FDragDropOperation> DragDropOperation, EItemDropZone InItemDropZone) const
+{
+	if(DragDropOperation->IsOfType<FHierarchyDragDropOp>() && SectionViewModelWeak.IsValid())
 	{
 		TSharedPtr<FHierarchyDragDropOp> HierarchyDragDropOp = StaticCastSharedPtr<FHierarchyDragDropOp>(DragDropOperation);
 		TSharedPtr<FHierarchyElementViewModel> DraggedItem = HierarchyDragDropOp->GetDraggedElement().Pin();
 		
-		if(SectionViewModel.IsValid())
-		{
-			return SectionViewModel->CanDropOn(DraggedItem, ItemDropZone).bCanPerform;
-		}
-		
-		// for the All section which has no valid view model, we simply do a check if the sections of the dragged categories are different
-		if(DragDropOperation->IsOfType<FHierarchyDragDropOp>())
-		{
-			if(const UHierarchyCategory* Category = Cast<UHierarchyCategory>(DraggedItem->GetData()))
-			{
-				return TryGetSectionData() != Category->GetSection();
-			}
-		}
+		return SectionViewModelWeak.Pin()->CanDropOn(DraggedItem, InItemDropZone).bCanPerform;
 	}
 
 	return false;
@@ -266,36 +232,14 @@ bool SHierarchySection::OnCanAcceptDrop(TSharedPtr<FDragDropOperation> DragDropO
 FReply SHierarchySection::OnDroppedOn(const FGeometry&, const FDragDropEvent& DragDropEvent, EItemDropZone DropZone) const
 {
 	bDraggedOn = false;
+	CurrentItemDropZone.Reset();
+
 	if(TSharedPtr<FHierarchyDragDropOp> HierarchyDragDropOp = DragDropEvent.GetOperationAs<FHierarchyDragDropOp>())
 	{
 		TSharedPtr<FHierarchyElementViewModel> DraggedItem = HierarchyDragDropOp->GetDraggedElement().Pin();
-		if(SectionViewModel.IsValid())
-		{
-			SectionViewModel->OnDroppedOn(DraggedItem, DropZone);
-			return FReply::Handled();
-		}
-		
-		if(UHierarchyCategory* Category = Cast<UHierarchyCategory>(DraggedItem->GetDataMutable()))
-		{
-			Category->SetSection(nullptr);
-			TArray<UHierarchyCategory*> ChildrenCategories;
-			Category->GetChildrenOfType<>(ChildrenCategories, true);
-
-			for(UHierarchyCategory* ChildCategory : ChildrenCategories)
-			{
-				ChildCategory->SetSection(nullptr);
-			}
-
-			// we only need to reparent if the parent isn't already the root. This stops unnecessary reordering
-			if(HierarchyDragDropOp->GetDraggedElement().Pin()->GetParent() != HierarchyViewModel->GetHierarchyRootViewModel())
-			{
-				HierarchyViewModel->GetHierarchyRootViewModel()->ReparentToThis(HierarchyDragDropOp->GetDraggedElement().Pin());
-			}
 		
-			HierarchyViewModel->RefreshHierarchyView();
-
-			return FReply::Handled();
-		}
+		SectionViewModelWeak.Pin()->OnDroppedOn(DraggedItem, DropZone);
+		return FReply::Handled();
 	}
 
 	return FReply::Unhandled();
@@ -303,18 +247,15 @@ FReply SHierarchySection::OnDroppedOn(const FGeometry&, const FDragDropEvent& Dr
 
 FReply SHierarchySection::OnMouseButtonDown(const FGeometry& MyGeometry, const FPointerEvent& MouseEvent)
 {
-	if(SectionViewModel.IsValid())
+	// we handle the event here so we can react on mouse button up
+	if(MouseEvent.IsMouseButtonDown(EKeys::RightMouseButton))
 	{
-		// we handle the event here so we can react on mouse button up
-		if(MouseEvent.IsMouseButtonDown(EKeys::RightMouseButton))
-		{
-			return FReply::Handled();
-		}
-		else if(MouseEvent.IsMouseButtonDown(EKeys::LeftMouseButton))
-		{			
-			OnSectionActivatedDelegate.ExecuteIfBound(SectionViewModel);
-			return FReply::Handled().DetectDrag(AsShared(), EKeys::LeftMouseButton).SetUserFocus(AsShared());	
-		}
+		return FReply::Handled();
+	}
+	else if(MouseEvent.IsMouseButtonDown(EKeys::LeftMouseButton))
+	{			
+		OnSectionActivatedDelegate.ExecuteIfBound(SectionViewModelWeak.Pin());
+		return FReply::Handled().DetectDrag(AsShared(), EKeys::LeftMouseButton).SetUserFocus(AsShared());	
 	}
 
 	return FReply::Unhandled();
@@ -322,7 +263,7 @@ FReply SHierarchySection::OnMouseButtonDown(const FGeometry& MyGeometry, const F
 
 FReply SHierarchySection::OnMouseButtonUp(const FGeometry& MyGeometry, const FPointerEvent& MouseEvent)
 {
-	if(SectionViewModel.IsValid() && SectionViewModel->IsForHierarchy())
+	if(SectionViewModelWeak.IsValid() && SectionViewModelWeak.Pin()->IsForHierarchy() && SectionViewModelWeak.Pin()->GetData() != nullptr)
 	{
 		if(MouseEvent.GetEffectingButton() == EKeys::RightMouseButton)
 		{			
@@ -335,7 +276,7 @@ FReply SHierarchySection::OnMouseButtonUp(const FGeometry& MyGeometry, const FPo
 				FPopupTransitionEffect(FPopupTransitionEffect::ContextMenu)
 			);
 				
-			OnSectionActivatedDelegate.ExecuteIfBound(SectionViewModel);			
+			OnSectionActivatedDelegate.ExecuteIfBound(SectionViewModelWeak.Pin());			
 			return FReply::Handled();
 		}
 	}
@@ -343,6 +284,19 @@ FReply SHierarchySection::OnMouseButtonUp(const FGeometry& MyGeometry, const FPo
 	return FReply::Unhandled();
 }
 
+FReply SHierarchySection::OnDrop(const FGeometry& MyGeometry, const FDragDropEvent& DragDropEvent)
+{
+	if(CurrentItemDropZone.IsSet())
+	{
+		if(OnCanAcceptDrop(DragDropEvent.GetOperation(), CurrentItemDropZone.GetValue()))
+		{
+			return OnDroppedOn(MyGeometry, DragDropEvent, CurrentItemDropZone.GetValue());
+		}
+	}
+
+	return FReply::Unhandled();
+}
+
 void SHierarchySection::OnDragEnter(const FGeometry& MyGeometry, const FDragDropEvent& DragDropEvent)
 {
 	bDraggedOn = true;
@@ -354,9 +308,9 @@ void SHierarchySection::OnDragEnter(const FGeometry& MyGeometry, const FDragDrop
 
 FReply SHierarchySection::OnDragDetected(const FGeometry& MyGeometry, const FPointerEvent& MouseEvent)
 {
-	if(MouseEvent.IsMouseButtonDown(EKeys::LeftMouseButton) && SectionViewModel.IsValid() && SectionViewModel->CanDrag().bCanPerform)
+	if(MouseEvent.IsMouseButtonDown(EKeys::LeftMouseButton) && SectionViewModelWeak.Pin()->CanDrag().bCanPerform)
 	{
-		TSharedRef<FSectionDragDropOp> SectionDragDropOp = MakeShared<FSectionDragDropOp>(SectionViewModel);
+		TSharedRef<FSectionDragDropOp> SectionDragDropOp = MakeShared<FSectionDragDropOp>(SectionViewModelWeak.Pin());
 		SectionDragDropOp->Construct();
 		return FReply::Handled().BeginDragDrop(SectionDragDropOp);
 	}
@@ -367,17 +321,45 @@ FReply SHierarchySection::OnDragDetected(const FGeometry& MyGeometry, const FPoi
 void SHierarchySection::OnDragLeave(const FDragDropEvent& DragDropEvent)
 {
 	bDraggedOn = false;
+	CurrentItemDropZone.Reset();
+	
 	if(TSharedPtr<FHierarchyDragDropOp> HierarchyDragDropOp = DragDropEvent.GetOperationAs<FHierarchyDragDropOp>())
 	{
 		HierarchyDragDropOp->SetDescription(FText::GetEmpty());
 	}
 }
 
+FReply SHierarchySection::OnDragOver(const FGeometry& MyGeometry, const FDragDropEvent& DragDropEvent)
+{
+	const FVector2f LocalPointerPos = MyGeometry.AbsoluteToLocal(DragDropEvent.GetScreenSpacePosition());
+	const EItemDropZone ItemHoverZone = ZoneFromPointerPosition(LocalPointerPos, MyGeometry.GetLocalSize());
+
+	if(TSharedPtr<FHierarchyDragDropOp> HierarchyDragDropOp = DragDropEvent.GetOperationAs<FHierarchyDragDropOp>())
+	{
+		TSharedPtr<FHierarchyElementViewModel> DraggedItem = HierarchyDragDropOp->GetDraggedElement().Pin();
+		FHierarchyElementViewModel::FCanPerformActionResults Results = SectionViewModelWeak.Pin()->CanDropOn(DraggedItem, ItemHoverZone);
+		HierarchyDragDropOp->SetDescription(Results.CanPerformMessage);
+
+		if(Results.bCanPerform)
+		{
+			CurrentItemDropZone = ItemHoverZone;
+		}
+		else
+		{
+			CurrentItemDropZone.Reset();
+		}
+
+		return FReply::Handled();
+	}
+	
+	return FReply::Unhandled();
+}
+
 FReply SHierarchySection::OnKeyDown(const FGeometry& MyGeometry, const FKeyEvent& InKeyEvent)
 {
-	if(SectionViewModel.IsValid() && InKeyEvent.GetKey() == EKeys::Delete && SectionViewModel->CanDelete())
+	if(SectionViewModelWeak.IsValid() && InKeyEvent.GetKey() == EKeys::Delete && SectionViewModelWeak.Pin()->CanDelete())
 	{
-		SectionViewModel->Delete();
+		SectionViewModelWeak.Pin()->Delete();
 		return FReply::Handled();
 	}
 
@@ -386,42 +368,44 @@ FReply SHierarchySection::OnKeyDown(const FGeometry& MyGeometry, const FKeyEvent
 
 void SHierarchySection::OnMouseLeave(const FPointerEvent& MouseEvent)
 {
-	//bDraggedOn = false;
 	SCompoundWidget::OnMouseLeave(MouseEvent);
 }
 
 UHierarchySection* SHierarchySection::TryGetSectionData() const
 {
-	return SectionViewModel.IsValid() ? Cast<UHierarchySection>(SectionViewModel->GetDataMutable()) : nullptr;
+	return Cast<UHierarchySection>(SectionViewModelWeak.Pin()->GetDataMutable());
 }
 
 FText SHierarchySection::GetText() const
 {
-	return SectionViewModel->GetSectionNameAsText();
+	return SectionViewModelWeak.IsValid() ? SectionViewModelWeak.Pin()->GetSectionNameAsText() : FText::GetEmpty();
 }
 
 FText SHierarchySection::GetTooltipText() const
 {
-	return SectionViewModel->GetSectionTooltip();
+	return SectionViewModelWeak.IsValid() ? SectionViewModelWeak.Pin()->GetSectionTooltip() : FText::GetEmpty();
 }
 
 void SHierarchySection::OnRenameSection(const FText& Text, ETextCommit::Type CommitType) const
 {
-	FScopedTransaction Transaction(LOCTEXT("Transaction_Rename_Section", "Renamed hierarchy section"));
-	HierarchyViewModel->GetHierarchyRoot()->Modify();
-	
-	SectionViewModel->Rename(FName(Text.ToString()));
+	if(SectionViewModelWeak.IsValid() && SectionViewModelWeak.Pin()->GetSectionNameAsText().EqualTo(Text) == false)
+	{
+		FScopedTransaction Transaction(LOCTEXT("Transaction_Rename_Section", "Renamed hierarchy section"));
+		HierarchyViewModel->GetHierarchyRoot()->Modify();
+		
+		SectionViewModelWeak.Pin()->Rename(FName(Text.ToString()));
+	}
 }
 
 bool SHierarchySection::OnVerifySectionRename(const FText& NewName, FText& OutTooltip) const
 {
 	// this function shouldn't be used in case the section isn't valid but we'll make sure regardless
-	if(!SectionViewModel.IsValid())
+	if(!SectionViewModelWeak.IsValid())
 	{
 		return false;
 	}
 
-	if(SectionViewModel->GetSectionName().ToString() != NewName.ToString())
+	if(SectionViewModelWeak.Pin()->GetSectionName().ToString() != NewName.ToString())
 	{
 		TArray<FString> SectionNames;
 
@@ -448,7 +432,7 @@ bool SHierarchySection::IsSectionSelected() const
 
 bool SHierarchySection::IsSectionReadOnly() const
 {
-	return SectionViewModel.IsValid() ? !SectionViewModel->CanRename() : true;
+	return SectionViewModelWeak.IsValid() ? !SectionViewModelWeak.Pin()->CanRename() : true;
 }
 
 ECheckBoxState SHierarchySection::GetSectionCheckState() const
@@ -458,7 +442,7 @@ ECheckBoxState SHierarchySection::GetSectionCheckState() const
 
 void SHierarchySection::OnSectionCheckChanged(ECheckBoxState NewState)
 {
-	OnSectionActivatedDelegate.ExecuteIfBound(SectionViewModel);
+	OnSectionActivatedDelegate.ExecuteIfBound(SectionViewModelWeak.Pin());
 }
 
 EActiveTimerReturnType SHierarchySection::ActivateSectionIfDragging(double CurrentTime, float DeltaTime) const
@@ -467,13 +451,44 @@ EActiveTimerReturnType SHierarchySection::ActivateSectionIfDragging(double Curre
 	{
 		if(IsSectionSelected() == false)
 		{
-			OnSectionActivatedDelegate.ExecuteIfBound(SectionViewModel);
+			OnSectionActivatedDelegate.ExecuteIfBound(SectionViewModelWeak.Pin());
 		}
 	}
 
 	return EActiveTimerReturnType::Stop;
 }
 
+const FSlateBrush* SHierarchySection::GetDropIndicatorBrush(EItemDropZone InItemDropZone) const
+{
+	switch (InItemDropZone)
+	{
+		case EItemDropZone::AboveItem: return FDataHierarchyEditorStyle::Get().GetBrush("HierarchyEditor.Drop.Section.Above"); break;
+		default:
+		case EItemDropZone::OntoItem: return FDataHierarchyEditorStyle::Get().GetBrush("HierarchyEditor.Drop.Section.Onto"); break;
+		case EItemDropZone::BelowItem: return FDataHierarchyEditorStyle::Get().GetBrush("HierarchyEditor.Drop.Section.Below"); break;
+	};
+}
+
+EItemDropZone SHierarchySection::ZoneFromPointerPosition(UE::Slate::FDeprecateVector2DParameter LocalPointerPos, UE::Slate::FDeprecateVector2DParameter LocalSize)
+{
+	const float PointerPos = LocalPointerPos.X;
+	const float Size = LocalSize.X;
+
+	const float ZoneBoundarySu = FMath::Clamp(Size * 0.25f, 3.0f, 10.0f);
+	if (PointerPos < ZoneBoundarySu)
+	{
+		return EItemDropZone::AboveItem;
+	}
+	else if (PointerPos > Size - ZoneBoundarySu)
+	{
+		return EItemDropZone::BelowItem;
+	}
+	else
+	{
+		return EItemDropZone::OntoItem;
+	}
+}
+
 void SDataHierarchyEditor::Construct(const FArguments& InArgs, TObjectPtr<UDataHierarchyViewModelBase> InHierarchyViewModel)
 {
 	HierarchyViewModel = InHierarchyViewModel;
@@ -756,6 +771,8 @@ void SDataHierarchyEditor::Construct(const FArguments& InArgs, TObjectPtr<UDataH
 		FIsActionChecked(), FIsActionButtonVisible::CreateSP(this, &SDataHierarchyEditor::CanNavigateToMatchingHierarchyElementFromSelectedSourceElement));
 	
 	HierarchyViewModel->ForceFullRefresh();
+	
+	SetActiveSourceSection(DefaultSourceSectionViewModel);
 }
 
 SDataHierarchyEditor::~SDataHierarchyEditor()
@@ -942,7 +959,7 @@ void SDataHierarchyEditor::RefreshSectionsView()
 
 	for (TSharedPtr<FHierarchySectionViewModel>& SourceSection : SourceRootViewModel->GetSectionViewModels())
 	{
-		TSharedPtr<SHierarchySection> SectionWidget = SNew(SHierarchySection, SourceSection, HierarchyViewModel)
+		TSharedPtr<SHierarchySection> SectionWidget = SNew(SHierarchySection, SourceSection)
 		.IsSectionActive_Lambda([this, SourceSection]()
 		{
 			return GetActiveSourceSection() == SourceSection ? ECheckBoxState::Checked : ECheckBoxState::Unchecked;
@@ -953,6 +970,7 @@ void SDataHierarchyEditor::RefreshSectionsView()
 		});
 		
 		SourceSectionBox->AddSlot()
+		.Padding(2.f)
 		.HAlign(HAlign_Center)
 		.VAlign(VAlign_Center)
 		[
@@ -962,16 +980,17 @@ void SDataHierarchyEditor::RefreshSectionsView()
 
 	if(SourceRootViewModel->GetSectionViewModels().Num() > 0)
 	{
-		TSharedPtr<SHierarchySection> DefaultSourceSection = SNew(SHierarchySection, nullptr, HierarchyViewModel)
+		DefaultSourceSectionViewModel = MakeShared<FHierarchySectionViewModel>(nullptr, SourceRootViewModel.ToSharedRef(), HierarchyViewModel);
+		
+		TSharedPtr<SHierarchySection> DefaultSourceSection = SNew(SHierarchySection, DefaultSourceSectionViewModel)
 		.IsSectionActive_Lambda([this]()
 		{
-			return GetActiveSourceSection() == nullptr ? ECheckBoxState::Checked : ECheckBoxState::Unchecked;
+			return GetActiveSourceSection() == DefaultSourceSectionViewModel ? ECheckBoxState::Checked : ECheckBoxState::Unchecked;
 		})
-		.OnSectionActivated(this, &SDataHierarchyEditor::SetActiveSourceSection)
-		// we forbid drop on here as the 'All' sections don't have a valid view model to determine it instead
-		.bForbidDropOn(true);
+		.OnSectionActivated(this, &SDataHierarchyEditor::SetActiveSourceSection);
 
 		SourceSectionBox->AddSlot()
+		.Padding(2.f)
 		.HAlign(HAlign_Center)
 		.VAlign(VAlign_Center)
 		[
@@ -981,7 +1000,7 @@ void SDataHierarchyEditor::RefreshSectionsView()
 	
 	for (TSharedPtr<FHierarchySectionViewModel>& HierarchySection : HierarchyViewModel->GetHierarchyRootViewModel()->GetSectionViewModels())
 	{
-		TSharedPtr<SHierarchySection> SectionWidget = SNew(SHierarchySection, HierarchySection, HierarchyViewModel)
+		TSharedPtr<SHierarchySection> SectionWidget = SNew(SHierarchySection, HierarchySection)
 		.IsSectionActive_Lambda([this, HierarchySection]()
 		{
 			return HierarchyViewModel->GetActiveHierarchySection() == HierarchySection ? ECheckBoxState::Checked : ECheckBoxState::Unchecked;
@@ -992,6 +1011,7 @@ void SDataHierarchyEditor::RefreshSectionsView()
 		});
 		
 		HierarchySectionBox->AddSlot()
+		.Padding(2.f)
 		.HAlign(HAlign_Center)
 		.VAlign(VAlign_Center)
 		[
@@ -999,10 +1019,10 @@ void SDataHierarchyEditor::RefreshSectionsView()
 		];
 	}
 	
-	TSharedPtr<SHierarchySection> DefaultHierarchySection = SNew(SHierarchySection, nullptr, HierarchyViewModel)
+	TSharedPtr<SHierarchySection> DefaultHierarchySection = SNew(SHierarchySection, HierarchyViewModel->GetDefaultHierarchySectionViewModel())
 	.IsSectionActive_Lambda([this]()
 	{
-		return HierarchyViewModel->GetActiveHierarchySection() == nullptr ? ECheckBoxState::Checked : ECheckBoxState::Unchecked;
+		return HierarchyViewModel->GetActiveHierarchySection() == HierarchyViewModel->GetDefaultHierarchySectionViewModel() ? ECheckBoxState::Checked : ECheckBoxState::Unchecked;
 	})
 	.OnSectionActivated_Lambda([this](TSharedPtr<FHierarchySectionViewModel> SectionViewModel)
 	{
@@ -1010,6 +1030,7 @@ void SDataHierarchyEditor::RefreshSectionsView()
 	});
 
 	HierarchySectionBox->AddSlot()
+	.Padding(2.f)
 	.HAlign(HAlign_Center)
 	.VAlign(VAlign_Center)
 	[
@@ -1161,7 +1182,14 @@ bool SDataHierarchyEditor::FilterForSourceSection(TSharedPtr<const FHierarchyEle
 {
 	if(ActiveSourceSection.IsValid())
 	{
-		return GetActiveSourceSectionData() == ItemViewModel->GetSection();
+		// If the currently selected section data is nullptr, it's the All section, and we let everything pass
+		if(ActiveSourceSection.Pin()->GetData() == nullptr)
+		{
+			return true;
+		}
+
+		// if not, we check against identical section data
+		return ActiveSourceSection.Pin()->GetData() == ItemViewModel->GetSection();
 	}
 
 	return true;
@@ -1379,12 +1407,14 @@ void SDataHierarchyEditor::OnElementAdded(TSharedPtr<FHierarchyElementViewModel>
 
 void SDataHierarchyEditor::OnHierarchySectionActivated(TSharedPtr<FHierarchySectionViewModel> Section)
 {
-	OnSelectionChanged(Section, ESelectInfo::Direct, true);
+	// We forward nullptr in case this is the 'All' default section. We determine this by checking its data for validity.
+	// The all section does not have an actual element associated with it.
+	OnSelectionChanged(Section.IsValid() && Section->GetData() ? Section : nullptr, ESelectInfo::Direct, true);
 }
 
 void SDataHierarchyEditor::OnSourceSectionActivated(TSharedPtr<FHierarchySectionViewModel> Section)
 {
-	OnSelectionChanged(Section, ESelectInfo::Direct, false);
+	OnSelectionChanged(Section.IsValid() && Section->GetData() ? Section : nullptr, ESelectInfo::Direct, false);
 	RunSourceSearch();
 }
 
@@ -1398,7 +1428,7 @@ void SDataHierarchyEditor::OnHierarchySectionDeleted(TSharedPtr<FHierarchySectio
 {
 	if(HierarchyViewModel->GetActiveHierarchySection() == DeletedSection)
 	{
-		HierarchyViewModel->SetActiveHierarchySection(nullptr);
+		HierarchyViewModel->SetActiveHierarchySection(HierarchyViewModel->GetDefaultHierarchySectionViewModel());
 	}
 }
 
@@ -1723,36 +1753,4 @@ void SDataHierarchyEditor::NotifyPostChange(const FPropertyChangedEvent& Propert
 	HierarchyViewModel->OnHierarchyPropertiesChanged().Broadcast();
 }
 
-void SSectionDragDropTarget::Construct(const FArguments& InArgs,TSharedPtr<FHierarchySectionViewModel> InOwningSection, EItemDropZone InItemDropZone)
-{
-	OwningSection = InOwningSection;
-	DropZone = InItemDropZone;
-	
-	SDropTarget::Construct(InArgs._DropTargetArgs);
-}
-
-FReply SSectionDragDropTarget::OnDragOver(const FGeometry& MyGeometry, const FDragDropEvent& DragDropEvent)
-{
-	if(TSharedPtr<FHierarchyDragDropOp> HierarchyDragDropOp = DragDropEvent.GetOperationAs<FHierarchyDragDropOp>())
-	{
-		TSharedPtr<FHierarchyElementViewModel> DraggedItem = HierarchyDragDropOp->GetDraggedElement().Pin();
-		FHierarchyElementViewModel::FCanPerformActionResults Results = OwningSection->CanDropOn(DraggedItem, DropZone);
-		HierarchyDragDropOp->SetDescription(Results.CanPerformMessage);
-	}
-
-	return SDropTarget::OnDragOver(MyGeometry, DragDropEvent);
-}
-
-void SSectionDragDropTarget::OnDragLeave(const FDragDropEvent& DragDropEvent)
-{
-	TSharedPtr<FHierarchyDragDropOp> HierarchyDragDropOp = DragDropEvent.GetOperationAs<FHierarchyDragDropOp>();
-
-	if(HierarchyDragDropOp.IsValid())
-	{
-		HierarchyDragDropOp->SetDescription(FText::GetEmpty());
-	}
-	
-	SDropTarget::OnDragLeave(DragDropEvent);
-}
-
 #undef LOCTEXT_NAMESPACE
diff --git a/Engine/Source/Editor/DataHierarchyEditor/Public/DataHierarchyViewModelBase.h b/Engine/Source/Editor/DataHierarchyEditor/Public/DataHierarchyViewModelBase.h
--- a/Engine/Source/Editor/DataHierarchyEditor/Public/DataHierarchyViewModelBase.h
+++ b/Engine/Source/Editor/DataHierarchyEditor/Public/DataHierarchyViewModelBase.h
@@ -461,6 +461,7 @@ public:
 	/** Hierarchy items reflect the already edited hierarchy. This should generally be constructed from persistent serialized data. */
 	const TArray<TSharedPtr<FHierarchyElementViewModel>>& GetHierarchyItems() const;
 
+	TSharedPtr<struct FHierarchySectionViewModel> GetDefaultHierarchySectionViewModel() const { return DefaultHierarchySectionViewModel; }
 public:
 	/** The hierarchy root the widget is editing. This should point to persistent data stored somewhere else as the serialized root of the hierarchy. */
 	virtual UHierarchyRoot* GetHierarchyRoot() const PURE_VIRTUAL(UDataHierarchyViewModelBase::GetHierarchyRoot, return nullptr;);
@@ -571,6 +572,7 @@ protected:
 	
 	TSharedPtr<struct FHierarchyRootViewModel> HierarchyRootViewModel;
 
+	TSharedPtr<FHierarchySectionViewModel> DefaultHierarchySectionViewModel;
 	TWeakPtr<struct FHierarchySectionViewModel> ActiveHierarchySection;
 
 	TSharedPtr<FUICommandList> Commands;
@@ -889,29 +891,41 @@ private:
 
 struct DATAHIERARCHYEDITOR_API FHierarchySectionViewModel : FHierarchyElementViewModel
 {
-	FHierarchySectionViewModel(UHierarchySection* InItem, TSharedRef<FHierarchyRootViewModel> InParent, TWeakObjectPtr<UDataHierarchyViewModelBase> InHierarchyViewModel) : FHierarchyElementViewModel(InItem, InParent, InHierarchyViewModel, InParent->IsForHierarchy()) {}
+	FHierarchySectionViewModel(UHierarchySection* InItem, TSharedRef<FHierarchyRootViewModel> InParent, TWeakObjectPtr<UDataHierarchyViewModelBase> InHierarchyViewModel) : FHierarchyElementViewModel(InItem, InParent, InHierarchyViewModel, InParent->IsForHierarchy())
+	{
+		if(bIsForHierarchy == false)
+		{
+			SetDropDisallowed(true);
+		}
+	}
 	
 	virtual ~FHierarchySectionViewModel() override {}
 
-	void SetSectionName(FName InSectionName) const { Cast<UHierarchySection>(Element)->SetSectionName(InSectionName); }
-	FName GetSectionName() const { return Cast<UHierarchySection>(Element)->GetSectionName(); }
-	
-	void SetSectionNameAsText(const FText& Text) const { Cast<UHierarchySection>(Element)->SetSectionNameAsText(Text); }
-	FText GetSectionNameAsText() const { return Cast<UHierarchySection>(Element)->GetSectionNameAsText(); }
-	FText GetSectionTooltip() const { return Cast<UHierarchySection>(Element)->GetTooltip(); }
+	virtual FString ToString() const override;
 	
+	void SetSectionName(FName InSectionName);
+	FName GetSectionName() const;
+
+	void SetSectionNameAsText(const FText& Text);
+	FText GetSectionNameAsText() const;
+	FText GetSectionTooltip() const;
+
 	void SetSectionImage(const FSlateBrush* InSectionImage) { SectionImage = InSectionImage; }
 	const FSlateBrush* GetSectionImage() const { return SectionImage; }
 
+	void SetDropDisallowed(bool bInDropDisallowed) { bDropDisallowed = bInDropDisallowed; }
 protected:
 	/** Only hierarchy sections are editable */
 	virtual FCanPerformActionResults IsEditableByUser() override { return FCanPerformActionResults(IsForHierarchy()); }
 	virtual bool CanHaveChildren() const override { return false; }
-	virtual FCanPerformActionResults CanDragInternal() override { return bIsForHierarchy; }
+	virtual FCanPerformActionResults CanDragInternal() override;
+	
 	/** We can only rename hierarchy sections */
-	virtual bool CanRenameInternal() override { return IsForHierarchy(); }
+	virtual bool CanRenameInternal() override;
 	virtual void RenameInternal(FName NewName) override { GetDataMutable<UHierarchySection>()->SetSectionName(NewName); }
 
+	virtual bool CanDeleteInternal() override;
+	
 	virtual FCanPerformActionResults CanDropOnInternal(TSharedPtr<FHierarchyElementViewModel> DraggedItem, EItemDropZone ItemDropZone) override;
 	virtual void OnDroppedOnInternal(TSharedPtr<FHierarchyElementViewModel> DroppedItem, EItemDropZone ItemDropZone) override;
 
@@ -919,6 +933,8 @@ protected:
 
 private:
 	const FSlateBrush* SectionImage = nullptr;
+	
+	bool bDropDisallowed = false;
 };
 
 struct DATAHIERARCHYEDITOR_API FHierarchyItemViewModel : FHierarchyElementViewModel
