- Fixed deletion being done via identities instead of directly on the view models. Categories don't always have valid identities, so deleting via identity on a category could result in the wrong category getting deleted.
- Deleted "delete via identities" function as it was unused. Deletion via identity should be reserved for special cases and individual items, not groups of items.
- Added identity construction for new categories based on guid (as names can change)
- PostLoad identity construction for categories with invalid identities

diff --git a/Engine/Source/Editor/DataHierarchyEditor/Private/DataHierarchyViewModelBase.cpp b/Engine/Source/Editor/DataHierarchyEditor/Private/DataHierarchyViewModelBase.cpp
--- a/Engine/Source/Editor/DataHierarchyEditor/Private/DataHierarchyViewModelBase.cpp
+++ b/Engine/Source/Editor/DataHierarchyEditor/Private/DataHierarchyViewModelBase.cpp
@@ -411,6 +411,25 @@ void UHierarchyCategory::FixupSectionLinkage()
 	}
 }
 
+FHierarchyElementIdentity UHierarchyCategory::ConstructIdentity()
+{
+	FHierarchyElementIdentity Identity;
+	Identity.Names.Add("Category");
+	Identity.Guids.Add(FGuid::NewGuid());
+	return Identity;
+}
+
+void UHierarchyCategory::PostLoad()
+{
+	Super::PostLoad();
+	
+	// Some categories were never initialized with a proper identity. We fix this up here.
+	if(Identity.IsValid() == false)
+	{
+		SetIdentity(UHierarchyCategory::ConstructIdentity());
+	}
+}
+
 void UHierarchySection::SetSectionNameAsText(const FText& Text)
 {
 	Section = FName(Text.ToString());
@@ -1732,7 +1751,8 @@ void UDataHierarchyViewModelBase::AddCategory(TSharedPtr<FHierarchyElementViewMo
 	GetHierarchyRoot()->Modify();
 
 	UClass* CategoryClass = GetCategoryDataClass();
-	UHierarchyCategory* Category = Cast<UHierarchyCategory>(CategoryParent->AddChild(CategoryClass, FHierarchyElementIdentity()));
+	
+	UHierarchyCategory* Category = Cast<UHierarchyCategory>(CategoryParent->AddChild(CategoryClass, UHierarchyCategory::ConstructIdentity()));
 	
 	TSharedPtr<FHierarchyElementViewModel> ViewModel = CategoryParent->FindViewModelForChild(Category, false);
 	if(ensureMsgf(ViewModel.IsValid(), TEXT("Could not find view model for new category of type '%s'. Please ensure your 'CreateViewModelForData' function creates a view model."), *CategoryClass->GetName()))
@@ -1807,6 +1827,11 @@ UHierarchyElement* UDataHierarchyViewModelBase::AddElementUnderRoot(TSubclassOf<
 
 void UDataHierarchyViewModelBase::DeleteElementWithIdentity(FHierarchyElementIdentity Identity)
 {
+	if(Identity.IsValid() == false)
+	{
+		return;
+	}
+	
 	FScopedTransaction Transaction(LOCTEXT("Transaction_DeleteItem", "Deleted hierarchy item"));
 	HierarchyRoot->Modify();
 
@@ -1835,30 +1860,10 @@ void UDataHierarchyViewModelBase::DeleteElementWithIdentity(FHierarchyElementIde
 		HierarchyRootViewModel->SyncViewModelsToData();
 		OnHierarchyChangedDelegate.Broadcast();
 	}
-}
-
-void UDataHierarchyViewModelBase::DeleteElementsWithIdentities(TArray<FHierarchyElementIdentity> Identities)
-{
-	TArray<TSharedPtr<FHierarchyElementViewModel>> ViewModelsToDelete;
-	
-	for(FHierarchyElementIdentity& Identity : Identities)
+	else
 	{
-		if(TSharedPtr<FHierarchyElementViewModel> ViewModel = HierarchyRootViewModel->FindViewModelForChild(Identity, true))
-		{
-			ViewModelsToDelete.Add(ViewModel);
-		}
-
-		TArray<TSharedPtr<FHierarchySectionViewModel>> SectionViewModels = HierarchyRootViewModel->GetSectionViewModels();
-		for(TSharedPtr<FHierarchySectionViewModel> SectionViewModel : SectionViewModels)
-		{
-			if(SectionViewModel->GetData()->GetPersistentIdentity() == Identity)
-			{
-				ViewModelsToDelete.Add(SectionViewModel);
-			}
-		}
+		Transaction.Cancel();
 	}
-
-	DeleteElements(ViewModelsToDelete);
 }
 
 void UDataHierarchyViewModelBase::DeleteElements(TArray<TSharedPtr<FHierarchyElementViewModel>> ViewModels) const
@@ -1881,6 +1886,10 @@ void UDataHierarchyViewModelBase::DeleteElements(TArray<TSharedPtr<FHierarchyEle
 		HierarchyRootViewModel->SyncViewModelsToData();
 		OnHierarchyChangedDelegate.Broadcast();
 	}
+	else
+	{
+		Transaction.Cancel();
+	}
 }
 
 void UDataHierarchyViewModelBase::NavigateToElementInHierarchy(const FHierarchyElementIdentity& HierarchyIdentity) const
diff --git a/Engine/Source/Editor/DataHierarchyEditor/Private/Widgets/SDataHierarchyEditor.cpp b/Engine/Source/Editor/DataHierarchyEditor/Private/Widgets/SDataHierarchyEditor.cpp
--- a/Engine/Source/Editor/DataHierarchyEditor/Private/Widgets/SDataHierarchyEditor.cpp
+++ b/Engine/Source/Editor/DataHierarchyEditor/Private/Widgets/SDataHierarchyEditor.cpp
@@ -1282,14 +1282,7 @@ void SDataHierarchyEditor::ClearSourceItems() const
 
 void SDataHierarchyEditor::DeleteItems(TArray<TSharedPtr<FHierarchyElementViewModel>> ItemsToDelete) const
 {
-	TArray<FHierarchyElementIdentity> DeletionIdentities;
-	
-	for(TSharedPtr<FHierarchyElementViewModel>& SelectedItem : ItemsToDelete)
-	{
-		DeletionIdentities.Add(SelectedItem->GetData()->GetPersistentIdentity());
-	}
-
-	HierarchyViewModel->DeleteElementsWithIdentities(DeletionIdentities);
+	HierarchyViewModel->DeleteElements(ItemsToDelete);
 }
 
 void SDataHierarchyEditor::DeleteSelectedHierarchyItems() const
diff --git a/Engine/Source/Editor/DataHierarchyEditor/Public/DataHierarchyViewModelBase.h b/Engine/Source/Editor/DataHierarchyEditor/Public/DataHierarchyViewModelBase.h
--- a/Engine/Source/Editor/DataHierarchyEditor/Public/DataHierarchyViewModelBase.h
+++ b/Engine/Source/Editor/DataHierarchyEditor/Public/DataHierarchyViewModelBase.h
@@ -364,6 +364,11 @@ public:
 	 *  This function expects the correct section with the same identity to exist already at the root level
 	 */
 	void FixupSectionLinkage();
+
+	static FHierarchyElementIdentity ConstructIdentity();
+
+protected:
+	virtual void PostLoad() override;
 private:
 	UPROPERTY()
 	FName Category;
@@ -503,9 +508,11 @@ public:
 	void AddCategory(TSharedPtr<FHierarchyElementViewModel> CategoryParent = nullptr) const;
 	void AddSection() const;
 
-	void DeleteElementWithIdentity(FHierarchyElementIdentity Identity);
-	void DeleteElementsWithIdentities(TArray<FHierarchyElementIdentity> Identities);
+	/** Delete all specified elements. */
 	void DeleteElements(TArray<TSharedPtr<FHierarchyElementViewModel>> ViewModels) const;
+	/** Special case for deleting a specific element based on its identity.
+	 *  Useful for externally removing an element from the hierarchy when you don't have access to the view model.*/
+	void DeleteElementWithIdentity(FHierarchyElementIdentity Identity);
 
 	void NavigateToElementInHierarchy(const FHierarchyElementIdentity& HierarchyIdentity) const;
 	void NavigateToElementInHierarchy(const TSharedRef<FHierarchyElementViewModel> HierarchyElement) const;
