Updated headers using UnrealCodeFixup to make sure dllstorage is on methods/staticvars instead of types. Part 5/n

diff --git a/Engine/Source/Editor/CurveEditor/Public/RichCurveEditorModel.h b/Engine/Source/Editor/CurveEditor/Public/RichCurveEditorModel.h
--- a/Engine/Source/Editor/CurveEditor/Public/RichCurveEditorModel.h
+++ b/Engine/Source/Editor/CurveEditor/Public/RichCurveEditorModel.h
@@ -17,6 +17,8 @@
 #include "UObject/WeakObjectPtr.h"
 #include "UObject/WeakObjectPtrTemplates.h"
 
+#define UE_API CURVEEDITOR_API
+
 class FCurveEditor;
 class IBufferedCurveModel;
 class UObject;
@@ -28,41 +30,41 @@ struct FKeyHandle;
 struct FKeyPosition;
 struct FRichCurve;
 
-class CURVEEDITOR_API FRichCurveEditorModel : public FCurveModel
+class FRichCurveEditorModel : public FCurveModel
 {
 public:
-	FRichCurveEditorModel(UObject* InOwner);
+	UE_API FRichCurveEditorModel(UObject* InOwner);
 
-	virtual const void* GetCurve() const override;
+	UE_API virtual const void* GetCurve() const override;
 
-	virtual void Modify() override;
+	UE_API virtual void Modify() override;
 
-	virtual void DrawCurve(const FCurveEditor& CurveEditor, const FCurveEditorScreenSpace& ScreenSpace, TArray<TTuple<double, double>>& InterpolatingPoints) const override;
-	virtual void GetKeys(double MinTime, double MaxTime, double MinValue, double MaxValue, TArray<FKeyHandle>& OutKeyHandles) const override;
-	virtual void GetKeyDrawInfo(ECurvePointType PointType, const FKeyHandle InKeyHandle, FKeyDrawInfo& OutDrawInfo) const override;
+	UE_API virtual void DrawCurve(const FCurveEditor& CurveEditor, const FCurveEditorScreenSpace& ScreenSpace, TArray<TTuple<double, double>>& InterpolatingPoints) const override;
+	UE_API virtual void GetKeys(double MinTime, double MaxTime, double MinValue, double MaxValue, TArray<FKeyHandle>& OutKeyHandles) const override;
+	UE_API virtual void GetKeyDrawInfo(ECurvePointType PointType, const FKeyHandle InKeyHandle, FKeyDrawInfo& OutDrawInfo) const override;
 
-	virtual void GetKeyPositions(TArrayView<const FKeyHandle> InKeys, TArrayView<FKeyPosition> OutKeyPositions) const override;
-	virtual void SetKeyPositions(TArrayView<const FKeyHandle> InKeys, TArrayView<const FKeyPosition> InKeyPositions, EPropertyChangeType::Type ChangeType) override;
+	UE_API virtual void GetKeyPositions(TArrayView<const FKeyHandle> InKeys, TArrayView<FKeyPosition> OutKeyPositions) const override;
+	UE_API virtual void SetKeyPositions(TArrayView<const FKeyHandle> InKeys, TArrayView<const FKeyPosition> InKeyPositions, EPropertyChangeType::Type ChangeType) override;
 
-	virtual void GetKeyAttributes(TArrayView<const FKeyHandle> InKeys, TArrayView<FKeyAttributes> OutAttributes) const override;
-	virtual void GetKeyAttributesExcludingAutoComputed(TArrayView<const FKeyHandle> InKeys, TArrayView<FKeyAttributes> OutAttributes) const override;
-	virtual void SetKeyAttributes(TArrayView<const FKeyHandle> InKeys, TArrayView<const FKeyAttributes> InAttributes, EPropertyChangeType::Type ChangeType = EPropertyChangeType::Unspecified) override;
+	UE_API virtual void GetKeyAttributes(TArrayView<const FKeyHandle> InKeys, TArrayView<FKeyAttributes> OutAttributes) const override;
+	UE_API virtual void GetKeyAttributesExcludingAutoComputed(TArrayView<const FKeyHandle> InKeys, TArrayView<FKeyAttributes> OutAttributes) const override;
+	UE_API virtual void SetKeyAttributes(TArrayView<const FKeyHandle> InKeys, TArrayView<const FKeyAttributes> InAttributes, EPropertyChangeType::Type ChangeType = EPropertyChangeType::Unspecified) override;
 
-	virtual void GetCurveAttributes(FCurveAttributes& OutCurveAttributes) const override;
-	virtual void SetCurveAttributes(const FCurveAttributes& InCurveAttributes) override;
-	virtual void GetTimeRange(double& MinTime, double& MaxTime) const override;
-	virtual void GetValueRange(double& MinValue, double& MaxValue) const override;
-	virtual int32 GetNumKeys() const override;
-	virtual void GetNeighboringKeys(const FKeyHandle InKeyHandle, TOptional<FKeyHandle>& OutPreviousKeyHandle, TOptional<FKeyHandle>& OutNextKeyHandle) const override;
-	virtual TPair<ERichCurveInterpMode, ERichCurveTangentMode> GetInterpolationMode(const double& InTime, ERichCurveInterpMode DefaultInterpolationMode, ERichCurveTangentMode DefaultTangentMode) const override;
+	UE_API virtual void GetCurveAttributes(FCurveAttributes& OutCurveAttributes) const override;
+	UE_API virtual void SetCurveAttributes(const FCurveAttributes& InCurveAttributes) override;
+	UE_API virtual void GetTimeRange(double& MinTime, double& MaxTime) const override;
+	UE_API virtual void GetValueRange(double& MinValue, double& MaxValue) const override;
+	UE_API virtual int32 GetNumKeys() const override;
+	UE_API virtual void GetNeighboringKeys(const FKeyHandle InKeyHandle, TOptional<FKeyHandle>& OutPreviousKeyHandle, TOptional<FKeyHandle>& OutNextKeyHandle) const override;
+	UE_API virtual TPair<ERichCurveInterpMode, ERichCurveTangentMode> GetInterpolationMode(const double& InTime, ERichCurveInterpMode DefaultInterpolationMode, ERichCurveTangentMode DefaultTangentMode) const override;
 
-	virtual bool Evaluate(double ProspectiveTime, double& OutValue) const override;
-	virtual void AddKeys(TArrayView<const FKeyPosition> InKeyPositions, TArrayView<const FKeyAttributes> InAttributes, TArrayView<TOptional<FKeyHandle>>* OutKeyHandles) override;
-	virtual void RemoveKeys(TArrayView<const FKeyHandle> InKeys, double InCurrentTime) override;
+	UE_API virtual bool Evaluate(double ProspectiveTime, double& OutValue) const override;
+	UE_API virtual void AddKeys(TArrayView<const FKeyPosition> InKeyPositions, TArrayView<const FKeyAttributes> InAttributes, TArrayView<TOptional<FKeyHandle>>* OutKeyHandles) override;
+	UE_API virtual void RemoveKeys(TArrayView<const FKeyHandle> InKeys, double InCurrentTime) override;
 
-	virtual void CreateKeyProxies(TArrayView<const FKeyHandle> InKeyHandles, TArrayView<UObject*> OutObjects) override;
+	UE_API virtual void CreateKeyProxies(TArrayView<const FKeyHandle> InKeyHandles, TArrayView<UObject*> OutObjects) override;
 
-	virtual TUniquePtr<IBufferedCurveModel> CreateBufferedCurveCopy() const override;
+	UE_API virtual TUniquePtr<IBufferedCurveModel> CreateBufferedCurveCopy() const override;
 
 	// Set a range to clamp key input values
 	void SetClampInputRange(TAttribute<TRange<double>> InClampInputRange) { ClampInputRange = InClampInputRange; }
@@ -83,21 +85,23 @@ private:
 };
 
 // Rich curve model operating on a raw curve ptr
-class CURVEEDITOR_API FRichCurveEditorModelRaw : public FRichCurveEditorModel
+class FRichCurveEditorModelRaw : public FRichCurveEditorModel
 {
 public:
-	FRichCurveEditorModelRaw(FRichCurve* InRichCurve, UObject* InOwner);
+	UE_API FRichCurveEditorModelRaw(FRichCurve* InRichCurve, UObject* InOwner);
 
-	virtual bool IsReadOnly() const override;
+	UE_API virtual bool IsReadOnly() const override;
 
-	void SetIsReadOnly(TAttribute<bool> InReadOnlyAttribute);
+	UE_API void SetIsReadOnly(TAttribute<bool> InReadOnlyAttribute);
 
 	// FRichCurveEditorModel interface
 	virtual bool IsValid() const override { return RichCurve != nullptr; }
-	virtual FRichCurve& GetRichCurve() override;
-	virtual const FRichCurve& GetReadOnlyRichCurve() const override;
+	UE_API virtual FRichCurve& GetRichCurve() override;
+	UE_API virtual const FRichCurve& GetReadOnlyRichCurve() const override;
 
 private:
 	FRichCurve* RichCurve;
 	TAttribute<bool> ReadOnlyAttribute;
-};
\ No newline at end of file
+};
+
+#undef UE_API
diff --git a/Engine/Source/Editor/CurveEditor/Public/SCurveEditorToolProperties.h b/Engine/Source/Editor/CurveEditor/Public/SCurveEditorToolProperties.h
--- a/Engine/Source/Editor/CurveEditor/Public/SCurveEditorToolProperties.h
+++ b/Engine/Source/Editor/CurveEditor/Public/SCurveEditorToolProperties.h
@@ -7,6 +7,8 @@
 #include "Widgets/DeclarativeSyntaxSupport.h"
 #include "Widgets/SCompoundWidget.h"
 
+#define UE_API CURVEEDITOR_API
+
 class FCurveEditor;
 class IStructureDetailsView;
 struct FPropertyChangedEvent;
@@ -14,7 +16,7 @@ struct FPropertyChangedEvent;
 /**
  * Widget to show curve editor tool extension properties
  */
-class CURVEEDITOR_API SCurveEditorToolProperties : public SCompoundWidget
+class SCurveEditorToolProperties : public SCompoundWidget
 {
 public:
 	SLATE_BEGIN_ARGS(SCurveEditorToolProperties)
@@ -24,13 +26,13 @@ public:
 
 	SCurveEditorToolProperties() : ToolId(FCurveEditorToolID::Unset()) {}
 
-	void Construct(const FArguments& InArgs, TSharedRef<FCurveEditor> InCurveEditor, FCurveEditorToolID ToolId);
+	UE_API void Construct(const FArguments& InArgs, TSharedRef<FCurveEditor> InCurveEditor, FCurveEditorToolID ToolId);
 
 	/**
 	 * Called when the curve editor tool changes
 	 * Sets the new tool ID and rebuilds the details view
 	 */
-	void OnToolChanged(FCurveEditorToolID NewToolId);
+	UE_API void OnToolChanged(FCurveEditorToolID NewToolId);
 
 	/** The details view for tool properties */
 	TSharedPtr<class IStructureDetailsView> GetStructureDetailsView() const { return DetailsView; }
@@ -41,16 +43,18 @@ private:
 	 * Rebuilds details view
 	 * Called either by the tool itself when it's properties change or when the curve editor tool changes
 	 */
-	void RebuildProperties();
+	UE_API void RebuildProperties();
 
 	/**
 	 * Called when the user enters in a new value for one of the properties
 	 * Calls OnToolOptionsUpdated on the current tool
 	 */
-	void OnFinishedChangingProperties(const FPropertyChangedEvent& PropertyChangedEvent);
+	UE_API void OnFinishedChangingProperties(const FPropertyChangedEvent& PropertyChangedEvent);
 
 private:
 	TWeakPtr<FCurveEditor> WeakCurveEditor;
 	TSharedPtr<IStructureDetailsView> DetailsView;
 	FCurveEditorToolID ToolId;
-};
\ No newline at end of file
+};
+
+#undef UE_API
diff --git a/Engine/Source/Editor/CurveEditor/Public/SCurveEditorView.h b/Engine/Source/Editor/CurveEditor/Public/SCurveEditorView.h
--- a/Engine/Source/Editor/CurveEditor/Public/SCurveEditorView.h
+++ b/Engine/Source/Editor/CurveEditor/Public/SCurveEditorView.h
@@ -20,6 +20,8 @@
 #include "CurveEditorSettings.h"
 #include "CurveDrawInfo.h"
 
+#define UE_API CURVEEDITOR_API
+
 class FCurveEditor;
 class FSlateRect;
 class FText;
@@ -114,54 +116,54 @@ private:
  *			* Auto sized (bAutoSize=true): views derive their height from the widget's desired size;
  *			* Stretched (bAutoSize=false): views will be stretched to the size of the parent panel with a sensible minimum height.
  */
-class CURVEEDITOR_API SCurveEditorView : public SCompoundWidget
+class SCurveEditorView : public SCompoundWidget
 {
 public:
 
 	/**
 	 * Default constructor
 	 */
-	SCurveEditorView();
-	~SCurveEditorView();
+	UE_API SCurveEditorView();
+	UE_API ~SCurveEditorView();
 
 	/**
 	 * Get the default screen space ultility that defines this view's input, output and pixel metrics
 	 */
-	FCurveEditorScreenSpace GetViewSpace() const;
+	UE_API FCurveEditorScreenSpace GetViewSpace() const;
 
 	/**
 	 * Get the screen space ultility that defines this view's input, output and pixel metrics for the specified axis combination
 	 */
-	FCurveEditorScreenSpace GetViewSpace(const FName& InHorizontalAxis, const FName& InVerticalAxis) const;
+	UE_API FCurveEditorScreenSpace GetViewSpace(const FName& InHorizontalAxis, const FName& InVerticalAxis) const;
 
 	/**
 	 * Get the screen space ultility that defines the specified curves's input, output and pixel metrics.
 	 * The resulting struct defines the transformation from this view's widget pixel space to the curve input/output space
 	 */
-	FCurveEditorScreenSpace GetCurveSpace(FCurveModelID CurveID) const;
+	UE_API FCurveEditorScreenSpace GetCurveSpace(FCurveModelID CurveID) const;
 
 	/** @return The transform used to translate view space (absolute key coordinates) to the curve space (e.g. 0-1 range in normalized view). */
-	FTransform2d GetViewToCurveTransform(FCurveModelID CurveID) const;
+	UE_API FTransform2d GetViewToCurveTransform(FCurveModelID CurveID) const;
 
 	/**
 	 * Retrieve the horizontal screen-space information for the specified axis
 	 */
-	FCurveEditorScreenSpaceH GetHorizontalAxisSpace(FCurveEditorViewAxisID ID) const;
+	UE_API FCurveEditorScreenSpaceH GetHorizontalAxisSpace(FCurveEditorViewAxisID ID) const;
 
 	/**
 	 * Retrieve the vertical screen-space information for the specified axis
 	 */
-	FCurveEditorScreenSpaceV GetVerticalAxisSpace(FCurveEditorViewAxisID ID) const;
+	UE_API FCurveEditorScreenSpaceV GetVerticalAxisSpace(FCurveEditorViewAxisID ID) const;
 
 	/**
 	 * Retrueve the axis ID assigned to the specified curve and orientation
 	 */
-	FCurveEditorViewAxisID GetAxisForCurve(FCurveModelID CurveID, ECurveEditorAxisOrientation AxisOrientation) const;
+	UE_API FCurveEditorViewAxisID GetAxisForCurve(FCurveModelID CurveID, ECurveEditorAxisOrientation AxisOrientation) const;
 
 	/**
 	 * Retrueve the axis associated with the specified ID and orientation
 	 */
-	TSharedPtr<FCurveEditorAxis> GetAxis(FCurveEditorViewAxisID ID, ECurveEditorAxisOrientation AxisOrientation) const;
+	UE_API TSharedPtr<FCurveEditorAxis> GetAxis(FCurveEditorViewAxisID ID, ECurveEditorAxisOrientation AxisOrientation) const;
 
 	/**
 	 * Check whether this view should auto size to fit its FixedHeight or child slot content
@@ -198,7 +200,7 @@ public:
 	/**
 	 * Retrieve this view's input bounds, accounting for any offsets or padding from the outer container
 	 */
-	void GetInputBounds(double& OutInputMin, double& OutInputMax) const;
+	UE_API void GetInputBounds(double& OutInputMin, double& OutInputMax) const;
 
 	/**
 	 * Retrieve this view's minimum output bound in view space
@@ -219,19 +221,19 @@ public:
 	/**
 	 * Set this view's output bounds
 	 */
-	void SetOutputBounds(double InOutputMin, double InOutputMax, FCurveEditorViewAxisID AxisID = FCurveEditorViewAxisID());
+	UE_API void SetOutputBounds(double InOutputMin, double InOutputMax, FCurveEditorViewAxisID AxisID = FCurveEditorViewAxisID());
 
 	/**
 	 * Set this view's input bounds
 	 */
-	void SetInputBounds(double InInputMin, double InInputMax, FCurveEditorViewAxisID AxisID = FCurveEditorViewAxisID());
+	UE_API void SetInputBounds(double InInputMin, double InInputMax, FCurveEditorViewAxisID AxisID = FCurveEditorViewAxisID());
 
 	/**
 	 * Zoom this view in or out around its center point
 	 *
 	 * @param Amount       The amount to zoom by horizontally and vertically as a factor of the current size
 	 */
-	void Zoom(const FVector2D& Amount);
+	UE_API void Zoom(const FVector2D& Amount);
 
 	/**
 	 * Zoom this view in or out around the specified point
@@ -240,10 +242,10 @@ public:
 	 * @param TimeOrigin   The time origin to zoom around
 	 * @param ValueOrigin  The value origin to zoom around
 	 */
-	void ZoomAround(const FVector2D& Amount, double InputOrigin, double OutputOrigin);
+	UE_API void ZoomAround(const FVector2D& Amount, double InputOrigin, double OutputOrigin);
 
 	/** This should be called every tick by an owning widget, to see if the cache is valid, which will then recreate it and invalidate widget*/
-	virtual void CheckCacheAndInvalidateIfNeeded();
+	UE_API virtual void CheckCacheAndInvalidateIfNeeded();
 
 	UE_DEPRECATED(5.3, "Use UpdateViewToTransformCurves(double InputMin, double InputMax) instead.")
 	virtual void UpdateViewToTransformCurves() {}
@@ -252,13 +254,13 @@ public:
 	virtual void UpdateViewToTransformCurves(double InputMin, double InputMax) {};
 
 	/** Frame the view vertially by the input and output bounds, peformaing any custom clipping as needed */
-	virtual void FrameVertical(double InOutputMin, double InOutputMax, FCurveEditorViewAxisID AxisID = FCurveEditorViewAxisID());
+	UE_API virtual void FrameVertical(double InOutputMin, double InOutputMax, FCurveEditorViewAxisID AxisID = FCurveEditorViewAxisID());
 
 	/** Frame the view horizontall by the input and output bounds, peformaing any custom clipping as needed */
-	virtual void FrameHorizontal(double InInputMin, double InInputMax, FCurveEditorViewAxisID AxisID = FCurveEditorViewAxisID());
+	UE_API virtual void FrameHorizontal(double InInputMin, double InInputMax, FCurveEditorViewAxisID AxisID = FCurveEditorViewAxisID());
 
 	/** Returns the cure editor associated with this view, or nullptr if the view is not specific to a curve editor */
-	TSharedPtr<FCurveEditor> GetCurveEditor() const;
+	UE_API TSharedPtr<FCurveEditor> GetCurveEditor() const;
 
 public:
 
@@ -334,26 +336,26 @@ public:
 	virtual void GetGridLinesY(TSharedRef<const FCurveEditor> CurveEditor, TArray<float>& MajorGridLines, TArray<float>& MinorGridLines, TArray<FText>* MajorGridLabels = nullptr) const {}
 
 	/** Request a new render from the retainer widget */
-	void RefreshRetainer();
+	UE_API void RefreshRetainer();
 
 	/** Update the custom axes if necessary */
-	void UpdateCustomAxes();
+	UE_API void UpdateCustomAxes();
 
 protected:
 	/** Gets info about the curves being drawn. Converts actual curves into an abstract series of lines/points/handles/etc. */
 	UE_DEPRECATED(5.6, "Instead use the CurveDrawParamsCache, use FCurveDrawParamsCache::GetCurveDrawParamsSynchronous to get Curve Draw Parms as with this function.")
-	void GetCurveDrawParams(TArray<FCurveDrawParams>& OutDrawParams);
+	UE_API void GetCurveDrawParams(TArray<FCurveDrawParams>& OutDrawParams);
 
 	/** Gets curve draw params  get the cache by calling a single curve */
 	UE_DEPRECATED(5.6, "Instead use GetCurveDrawParamsCache, use FCurveDrawParamsCache::GetCurveDrawParamSynchronous to get Curve Draw Parm as with this function.")
-	void GetCurveDrawParam(TSharedPtr<FCurveEditor>& CurveEditor, const FCurveModelID& ModelID, FCurveModel* CurveModel, FCurveDrawParams& OutDrawParam) const;
+	UE_API void GetCurveDrawParam(TSharedPtr<FCurveEditor>& CurveEditor, const FCurveModelID& ModelID, FCurveModel* CurveModel, FCurveDrawParams& OutDrawParam) const;
 
 	/** Update all the curve view transforms from curve models */
-	void UpdateCurveViewTransformsFromModels();
+	UE_API void UpdateCurveViewTransformsFromModels();
 
 	// ~SWidget interface
-	virtual FVector2D ComputeDesiredSize(float LayoutScaleMultiplier) const override;
-	virtual void Tick(const FGeometry& AllottedGeometry, const double InCurrentTime, const float InDeltaTime) override;
+	UE_API virtual FVector2D ComputeDesiredSize(float LayoutScaleMultiplier) const override;
+	UE_API virtual void Tick(const FGeometry& AllottedGeometry, const double InCurrentTime, const float InDeltaTime) override;
 
 private:
 
@@ -361,10 +363,10 @@ private:
 	friend struct FCurveEditorPanelViewTracker;
 
 	/** Add a curve to this view */
-	void AddCurve(FCurveModelID CurveID);
+	UE_API void AddCurve(FCurveModelID CurveID);
 
 	/** Add a curve from this view */
-	void RemoveCurve(FCurveModelID CurveID);
+	UE_API void RemoveCurve(FCurveModelID CurveID);
 
 public:
 
@@ -503,3 +505,5 @@ public:
 		RetainerWidget = InWidget;
 	}
 };
+
+#undef UE_API
diff --git a/Engine/Source/Editor/CurveEditor/Public/SCurveKeyDetailPanel.h b/Engine/Source/Editor/CurveEditor/Public/SCurveKeyDetailPanel.h
--- a/Engine/Source/Editor/CurveEditor/Public/SCurveKeyDetailPanel.h
+++ b/Engine/Source/Editor/CurveEditor/Public/SCurveKeyDetailPanel.h
@@ -7,6 +7,8 @@
 #include "Widgets/DeclarativeSyntaxSupport.h"
 #include "Widgets/SCompoundWidget.h"
 
+#define UE_API CURVEEDITOR_API
+
 class FCurveEditor;
 // Forward Declare
 class IPropertyRowGenerator;
@@ -15,7 +17,7 @@ class SWidget;
 /**
  * Inline details panel that lets you edit the Time and Value of a generic FCurveEditor Key
  */
-class CURVEEDITOR_API SCurveKeyDetailPanel : public SCompoundWidget
+class SCurveKeyDetailPanel : public SCompoundWidget
 {
 public:
 	SLATE_BEGIN_ARGS(SCurveKeyDetailPanel)
@@ -23,16 +25,18 @@ public:
 
 	SLATE_END_ARGS()
 
-	void Construct(const FArguments& InArgs, TSharedRef<FCurveEditor> InCurveEditor);
+	UE_API void Construct(const FArguments& InArgs, TSharedRef<FCurveEditor> InCurveEditor);
 	TSharedPtr<IPropertyRowGenerator> GetPropertyRowGenerator() const { return PropertyRowGenerator; }
 
 private:
-	void PropertyRowsRefreshed();
-	void ConstructChildLayout(TSharedPtr<SWidget> TimeWidget, TSharedPtr<SWidget> ValueWidget);
+	UE_API void PropertyRowsRefreshed();
+	UE_API void ConstructChildLayout(TSharedPtr<SWidget> TimeWidget, TSharedPtr<SWidget> ValueWidget);
 
 private:
 	TSharedPtr<IPropertyRowGenerator> PropertyRowGenerator;
 
 	TSharedPtr<SWidget> TempTimeWidget;
 	TSharedPtr<SWidget> TempValueWidget;
-};
\ No newline at end of file
+};
+
+#undef UE_API
diff --git a/Engine/Source/Editor/CurveEditor/Public/SCurveViewerPanel.h b/Engine/Source/Editor/CurveEditor/Public/SCurveViewerPanel.h
--- a/Engine/Source/Editor/CurveEditor/Public/SCurveViewerPanel.h
+++ b/Engine/Source/Editor/CurveEditor/Public/SCurveViewerPanel.h
@@ -14,6 +14,8 @@
 #include "Widgets/DeclarativeSyntaxSupport.h"
 #include "Widgets/SLeafWidget.h"
 
+#define UE_API CURVEEDITOR_API
+
 class FCurveEditor;
 class FPaintArgs;
 class FSlateRect;
@@ -24,7 +26,7 @@ struct FGeometry;
 /**
  * Curve viewer widget that reflects the state of an FCurveEditor
  */
-class CURVEEDITOR_API SCurveViewerPanel : public SCurveEditorView
+class SCurveViewerPanel : public SCurveEditorView
 {
 	SLATE_BEGIN_ARGS(SCurveViewerPanel) {}
 
@@ -35,7 +37,7 @@ class CURVEEDITOR_API SCurveViewerPanel : public SCurveEditorView
 	/**
 	 * Construct a new curve editor panel widget
 	 */
-	void Construct(const FArguments& InArgs, TSharedRef<FCurveEditor> InCurveEditor);
+	UE_API void Construct(const FArguments& InArgs, TSharedRef<FCurveEditor> InCurveEditor);
 
 	/**
 	 * Access the draw parameters that this curve editor has cached for this frame
@@ -47,15 +49,15 @@ class CURVEEDITOR_API SCurveViewerPanel : public SCurveEditorView
 
 private:
 
-	virtual int32 OnPaint(const FPaintArgs& Args, const FGeometry& AllottedGeometry, const FSlateRect& MyCullingRect, FSlateWindowElementList& OutDrawElements, int32 LayerId, const FWidgetStyle& InWidgetStyle, bool bParentEnabled) const override;
-	virtual void Tick(const FGeometry& AllottedGeometry, const double InCurrentTime, const float InDeltaTime) override;
+	UE_API virtual int32 OnPaint(const FPaintArgs& Args, const FGeometry& AllottedGeometry, const FSlateRect& MyCullingRect, FSlateWindowElementList& OutDrawElements, int32 LayerId, const FWidgetStyle& InWidgetStyle, bool bParentEnabled) const override;
+	UE_API virtual void Tick(const FGeometry& AllottedGeometry, const double InCurrentTime, const float InDeltaTime) override;
 
 public:
 
 	/**
 	 * Draw curve data
 	 */
-	void DrawCurves(const FGeometry& AllottedGeometry, FSlateWindowElementList& OutDrawElements, int32 LayerId, const FWidgetStyle& InWidgetStyle, ESlateDrawEffect DrawEffects) const;
+	UE_API void DrawCurves(const FGeometry& AllottedGeometry, FSlateWindowElementList& OutDrawElements, int32 LayerId, const FWidgetStyle& InWidgetStyle, ESlateDrawEffect DrawEffects) const;
 
 private:
 	/** Curve draw parameters that are re-generated on tick */
@@ -63,4 +65,6 @@ private:
 
 	/** Thickness of the displayed curve */
 	TAttribute<float> CurveThickness;
-};
\ No newline at end of file
+};
+
+#undef UE_API
diff --git a/Engine/Source/Editor/CurveEditor/Public/Tree/CurveEditorTree.h b/Engine/Source/Editor/CurveEditor/Public/Tree/CurveEditorTree.h
--- a/Engine/Source/Editor/CurveEditor/Public/Tree/CurveEditorTree.h
+++ b/Engine/Source/Editor/CurveEditor/Public/Tree/CurveEditorTree.h
@@ -15,6 +15,8 @@
 #include "Templates/SharedPointer.h"
 #include "Templates/TypeHash.h"
 
+#define UE_API CURVEEDITOR_API
+
 enum class ECurveEditorTreeFilterType : uint32;
 
 class FCurveEditor;
@@ -47,13 +49,13 @@ ENUM_CLASS_FLAGS(ECurveEditorTreeFilterState);
 /**
  * Scoped guard that prevents the broadcast of tree events for the duration of its lifetime. Will trigger necessary events after the last remaining guard has been destroyed.
  */
-struct CURVEEDITOR_API FScopedCurveEditorTreeEventGuard
+struct FScopedCurveEditorTreeEventGuard
 {
-	explicit FScopedCurveEditorTreeEventGuard(FCurveEditorTree* InTree);
-	~FScopedCurveEditorTreeEventGuard();
+	UE_API explicit FScopedCurveEditorTreeEventGuard(FCurveEditorTree* InTree);
+	UE_API ~FScopedCurveEditorTreeEventGuard();
 
-	FScopedCurveEditorTreeEventGuard(FScopedCurveEditorTreeEventGuard&& RHS);
-	FScopedCurveEditorTreeEventGuard& operator=(FScopedCurveEditorTreeEventGuard&& RHS);
+	UE_API FScopedCurveEditorTreeEventGuard(FScopedCurveEditorTreeEventGuard&& RHS);
+	UE_API FScopedCurveEditorTreeEventGuard& operator=(FScopedCurveEditorTreeEventGuard&& RHS);
 
 private:
 	/** OnItemsChanged.SerialNumber cached on construction. Used to detect whether OnItemsChanged should be broadcast. */
@@ -126,7 +128,7 @@ struct FSortedCurveEditorTreeItems
  * Concrete type used as a tree item for the curve editor. No need to derive from this type - custom behaviour is implemented through ICurveEditorTreeItem.
  * Implemented in this way to ensure that all hierarchical information can be reasoned about within the curve editor itself, and allow for mixing of tree item types from any usage domain.
  */
-struct CURVEEDITOR_API FCurveEditorTreeItem
+struct FCurveEditorTreeItem
 {
 	/** @return This item's unique identifier within the tree */
 	FCurveEditorTreeItemID GetID() const
@@ -198,17 +200,17 @@ struct CURVEEDITOR_API FCurveEditorTreeItem
 	/**
 	 * Retrieve all the curves for this tree item, creating them through ICurveEditorTreeItem::CreateCurveModels if there are none
 	 */
-	TArrayView<const FCurveModelID> GetOrCreateCurves(FCurveEditor* CurveEditor);
+	UE_API TArrayView<const FCurveModelID> GetOrCreateCurves(FCurveEditor* CurveEditor);
 
 	/**
 	 * Destroy any previously constructed curve models that this tree item owns
 	 */
-	void DestroyCurves(FCurveEditor* CurveEditor);
+	UE_API void DestroyCurves(FCurveEditor* CurveEditor);
 
 	/**
 	 * Destroy any previously constructed unpinned curve models that this tree item owns
 	 */
-	void DestroyUnpinnedCurves(FCurveEditor* CurveEditor);
+	UE_API void DestroyUnpinnedCurves(FCurveEditor* CurveEditor);
 
 private:
 
@@ -359,7 +361,7 @@ DECLARE_MULTICAST_DELEGATE_OneParam(FOnCurveEditorToggleExpansionState, bool);
 /** 
  * Complete implementation of a curve editor tree. Only really defines the hierarchy and selection states for tree items.
  */
-class CURVEEDITOR_API FCurveEditorTree
+class FCurveEditorTree
 {
 public:
 	/** Defines a predicate for sorting curve editor tree item implementations. */
@@ -368,44 +370,44 @@ public:
 	/** Structure containing all the events for this tree */
 	FCurveEditorTreeEvents Events;
 
-	FCurveEditorTree();
+	UE_API FCurveEditorTree();
 
 	/**
 	 * Retrieve an item from its ID, assuming it is definitely valid
 	 */
-	FCurveEditorTreeItem& GetItem(FCurveEditorTreeItemID ItemID);
+	UE_API FCurveEditorTreeItem& GetItem(FCurveEditorTreeItemID ItemID);
 
 	/**
 	 * Retrieve an item from its ID, assuming it is definitely valid
 	 */
-	const FCurveEditorTreeItem& GetItem(FCurveEditorTreeItemID ItemID) const;
+	UE_API const FCurveEditorTreeItem& GetItem(FCurveEditorTreeItemID ItemID) const;
 
 	/**
 	 * Retrieve an item from its ID or nullptr if the ID is not valid
 	 */
-	FCurveEditorTreeItem* FindItem(FCurveEditorTreeItemID ItemID);
+	UE_API FCurveEditorTreeItem* FindItem(FCurveEditorTreeItemID ItemID);
 
 	/**
 	 * Retrieve an item from its ID or nullptr if the ID is not valid
 	 */
-	const FCurveEditorTreeItem* FindItem(FCurveEditorTreeItemID ItemID) const;
+	UE_API const FCurveEditorTreeItem* FindItem(FCurveEditorTreeItemID ItemID) const;
 
 	/** 
 	 * Retrieve this curve editor's root items irrespective of filter state
 	 */
-	const TArray<FCurveEditorTreeItemID>& GetRootItems() const;
+	UE_API const TArray<FCurveEditorTreeItemID>& GetRootItems() const;
 
 	/** 
 	 * Retrieve all the items stored in this tree irrespective of filter state
 	 */
-	const TMap<FCurveEditorTreeItemID, FCurveEditorTreeItem>& GetAllItems() const;
+	UE_API const TMap<FCurveEditorTreeItemID, FCurveEditorTreeItem>& GetAllItems() const;
 
 	/**
 	 * Add a new empty item to the tree
 	 *
 	 * @param ParentID The ID of the desired parent for the new item, or FCurveEditorTreeItemID::Invalid for root nodes
 	 */
-	FCurveEditorTreeItem* AddItem(FCurveEditorTreeItemID ParentID);
+	UE_API FCurveEditorTreeItem* AddItem(FCurveEditorTreeItemID ParentID);
 
 	/**
 	 * Remove an item and all its children from this tree, destroying any curves it may have created.
@@ -413,36 +415,36 @@ public:
 	 * @param ItemID The ID of the item to remove
 	 * @param CurveEditor (required) Pointer to the curve editor that owns this tree to remove curves from
 	 */
-	void RemoveItem(FCurveEditorTreeItemID ItemID, FCurveEditor* CurveEditor);
+	UE_API void RemoveItem(FCurveEditorTreeItemID ItemID, FCurveEditor* CurveEditor);
 
 	/**
 	 * Run all the filters on this tree, updating filter state for all tree items
 	 */
-	void RunFilters();
+	UE_API void RunFilters();
 
 	/**
 	 * Add a new filter to this tree. Does not run the filter (and thus update any tree views) until RunFilters is called.
 	 *
 	 * @param NewFilter The new filter to add to this tree
 	 */
-	void AddFilter(TWeakPtr<FCurveEditorTreeFilter> NewFilter);
+	UE_API void AddFilter(TWeakPtr<FCurveEditorTreeFilter> NewFilter);
 
 	/**
 	 * Remove an existing filter from this tree. Does not re-run the filters (and thus update any tree views) until RunFilters is called.
 	 *
 	 * @param FilterToRemove The filter to remove from this tree
 	 */
-	void RemoveFilter(TWeakPtr<FCurveEditorTreeFilter> FilterToRemove);
+	UE_API void RemoveFilter(TWeakPtr<FCurveEditorTreeFilter> FilterToRemove);
 
 	/**
 	 * Direct access to all current filters (potentially including expired ones)
 	 */
-	TArrayView<const TWeakPtr<FCurveEditorTreeFilter>> GetFilters() const;
+	UE_API TArrayView<const TWeakPtr<FCurveEditorTreeFilter>> GetFilters() const;
 
 	/**
 	 * Clear all filters from the tree
 	 */
-	void ClearFilters();
+	UE_API void ClearFilters();
 
 	/**
 	 * Attempt to locate a filter by its type
@@ -450,48 +452,48 @@ public:
 	 * @param Type The type of the filter to find
 	 * @return A pointer to the filter, or nullptr if one could not be found
 	 */
-	const FCurveEditorTreeFilter* FindFilterByType(ECurveEditorTreeFilterType Type) const;
+	UE_API const FCurveEditorTreeFilter* FindFilterByType(ECurveEditorTreeFilterType Type) const;
 
 	/**
 	 * Sets a predicate which will be used to sort tree items after they're been marked as needing sort.
 	 */
-	 void SetSortPredicate(FTreeItemSortPredicate InSortPredicate);
+	 UE_API void SetSortPredicate(FTreeItemSortPredicate InSortPredicate);
 
 	 /**
 	  * Sorts all tree items which have been marked for sorting if the sort predicate has been set.
 	  */
-	void SortTreeItems();
+	UE_API void SortTreeItems();
 
 	/**
 	 * Inform this tree that the specified tree item IDs have been directly selected on the UI.
 	 * @note: This populates both implicit and explicit selection state for the supplied items and any children/parents
 	 */
-	void SetDirectSelection(TArray<FCurveEditorTreeItemID>&& TreeItems, FCurveEditor* InCurveEditor);
+	UE_API void SetDirectSelection(TArray<FCurveEditorTreeItemID>&& TreeItems, FCurveEditor* InCurveEditor);
 
 	/**
 	 * Removes tree items from the current selection.
 	 */
-	void RemoveFromSelection(TArrayView<const FCurveEditorTreeItemID> TreeItems, FCurveEditor* InCurveEditor);
+	UE_API void RemoveFromSelection(TArrayView<const FCurveEditorTreeItemID> TreeItems, FCurveEditor* InCurveEditor);
 
 	/**
 	 * Access the selection state for this tree. Items that are neither implicitly or explicitly selected are not present in the map.
 	 */
-	const TMap<FCurveEditorTreeItemID, ECurveEditorTreeSelectionState>& GetSelection() const;
+	UE_API const TMap<FCurveEditorTreeItemID, ECurveEditorTreeSelectionState>& GetSelection() const;
 
 	/**
 	 * Check a specific tree item's selection state
 	 */
-	ECurveEditorTreeSelectionState GetSelectionState(FCurveEditorTreeItemID InTreeItemID) const;
+	UE_API ECurveEditorTreeSelectionState GetSelectionState(FCurveEditorTreeItemID InTreeItemID) const;
 
 	/**
 	 * Access the filter state for this tree. Items that are neither implicitly or explicitly filtered-in are not present in the map.
 	 */
-	const FCurveEditorFilterStates& GetFilterStates() const;
+	UE_API const FCurveEditorFilterStates& GetFilterStates() const;
 
 	/**
 	 * Check a specific tree item's filter state
 	 */
-	ECurveEditorTreeFilterState GetFilterState(FCurveEditorTreeItemID InTreeItemID) const;
+	UE_API ECurveEditorTreeFilterState GetFilterState(FCurveEditorTreeItemID InTreeItemID) const;
 
 	/**
 	 * Retrieve a scoped event guard that will block broadcast of events until the last guard on the stack goes out of scope
@@ -505,12 +507,12 @@ public:
 	/**
 	 * Compact the memory used by this tree (does not modify any meaningful state)
 	 */
-	void Compact();
+	UE_API void Compact();
 
 	/*
 	 * Toggle the expansion state of the selected nodes or all nodes if none selected
 	 */
-	void ToggleExpansionState(bool bRecursive);
+	UE_API void ToggleExpansionState(bool bRecursive);
 
 	FOnCurveEditorToggleExpansionState& GetToggleExpansionState()
 	{ 
@@ -520,12 +522,12 @@ public:
 	/*
 	* Get cached expanded items
 	*/
-	TArray<FCurveEditorTreeItemID> GetCachedExpandedItems() const;
+	UE_API TArray<FCurveEditorTreeItemID> GetCachedExpandedItems() const;
 	
 	/*
 	* Set Item expansion state
 	*/
-	void SetItemExpansion(FCurveEditorTreeItemID InTreeItemID, bool bInExpansion);
+	UE_API void SetItemExpansion(FCurveEditorTreeItemID InTreeItemID, bool bInExpansion);
 
 	/**
 	Whether or not we are are doign a direct selection, could be used to see why a curve model is being created or destroyed, by direct selection or by sequencer filtering?
@@ -538,12 +540,12 @@ public:
 	/**
 	Recreate the curve models from the existing selection, this may be needed in case of a setting change.
 	*/
-	void RecreateModelsFromExistingSelection(FCurveEditor* CurveEditor);
+	UE_API void RecreateModelsFromExistingSelection(FCurveEditor* CurveEditor);
 
 private:
 
 	// Recursively removes children without removing them from the parent (assuming the parent is also being removed)
-	void RemoveChildrenRecursive(TArray<FCurveEditorTreeItemID>&& Children, FCurveEditor* CurveEditor);
+	UE_API void RemoveChildrenRecursive(TArray<FCurveEditorTreeItemID>&& Children, FCurveEditor* CurveEditor);
 
 	/**
 	 * Run the specified filters over the specified items and their recursive children, storing the results in this instance's FilterStates struct.
@@ -553,12 +555,12 @@ private:
 	 * @param InheritedState The filter state for each item to receive if it does not directly match a filter (either ECurveEditorTreeFilterState::NoMatch or ECurveEditorTreeFilterState::InheritedChild)
 	 * @return Raised state flags which should be applied to parents of the filtered items.
 	 */
-	ECurveEditorTreeFilterState PerformFilterPass(TArrayView<const FCurveEditorTreeFilter* const> FilterPtrs, TArrayView<const FCurveEditorTreeItemID> ItemsToFilter, ECurveEditorTreeFilterState InheritedState);
+	UE_API ECurveEditorTreeFilterState PerformFilterPass(TArrayView<const FCurveEditorTreeFilter* const> FilterPtrs, TArrayView<const FCurveEditorTreeItemID> ItemsToFilter, ECurveEditorTreeFilterState InheritedState);
 
 	/** 
 	 * Recursively sorts the tree item ids using the sort predicate.
 	 */
-	void SortTreeItems(FSortedCurveEditorTreeItems& TreeItemIDsToSort);
+	UE_API void SortTreeItems(FSortedCurveEditorTreeItems& TreeItemIDsToSort);
 
 	/** Incrementing ID for the next tree item to be created */
 	FCurveEditorTreeItemID NextTreeItemID;
@@ -590,4 +592,6 @@ private:
 
 	/** Whether or not we are are doign a direct selection, could be used to see why a curve model is being created or destroyed, by direct selection or by sequencer filtering?*/
 	bool bIsDoingDirectSelection = false;
-};
\ No newline at end of file
+};
+
+#undef UE_API
diff --git a/Engine/Source/Editor/CurveEditor/Public/Tree/CurveEditorTreeFilter.h b/Engine/Source/Editor/CurveEditor/Public/Tree/CurveEditorTreeFilter.h
--- a/Engine/Source/Editor/CurveEditor/Public/Tree/CurveEditorTreeFilter.h
+++ b/Engine/Source/Editor/CurveEditor/Public/Tree/CurveEditorTreeFilter.h
@@ -14,6 +14,8 @@
 #include "Templates/SharedPointer.h"
 #include "Templates/UnrealTemplate.h"
 
+#define UE_API CURVEEDITOR_API
+
 class FCurveEditorTree;
 struct FCurveEditorFilterStates;
 struct FCurveEditorTreeItem;
@@ -37,7 +39,7 @@ enum class ECurveEditorTreeFilterType : uint32
  * Filters are identifyable through their type (see GetType()), which is a pre-registered static value retrieved through RegisterFilterType(). Client types can inspect this for supported types, and static_cast<> for implementation details.
  * Filters also contain a pass, which is used to separate filters into separate passes. Filters applied within the same pass are matched as a boolean OR, filters in different passes are matched as a boolean AND.
  */
-struct CURVEEDITOR_API FCurveEditorTreeFilter
+struct FCurveEditorTreeFilter
 {
 	FCurveEditorTreeFilter(ECurveEditorTreeFilterType InFilterType, int32 InFilterPass)
 		: FilterType(InFilterType)
@@ -79,7 +81,7 @@ public:
 	 * @param InTreeItem The tree item to test
 	 * @return true if the item matches the filter, false otherwise
 	 */
-	static ECurveEditorTreeFilterType RegisterFilterType();
+	static UE_API ECurveEditorTreeFilterType RegisterFilterType();
 
 protected:
 
@@ -161,7 +163,7 @@ struct FCurveEditorTreeTextFilterTerm
 /**
  * Built-in text filter of type ECurveEditorTreeFilterType::Text. Filter terms are applied as a case-insensitive boolean OR substring match.
  */
-struct CURVEEDITOR_API FCurveEditorTreeTextFilter : FCurveEditorTreeFilter
+struct FCurveEditorTreeTextFilter : FCurveEditorTreeFilter
 {
 	/** Original input filter text. */
 	FText InputText;
@@ -179,7 +181,7 @@ struct CURVEEDITOR_API FCurveEditorTreeTextFilter : FCurveEditorTreeFilter
 	 *
 	 * @param FilterString   The new filter string that should be parsed into this filter
 	 */
-	void AssignFromText(const FString& FilterString);
+	UE_API void AssignFromText(const FString& FilterString);
 
 	/**
 	 * Check whether this filter is empty (ie, has no filter terms)
@@ -202,4 +204,6 @@ private:
 
 	/** Arrays of case-insensitive terms to find within tree items. Each term is an array of strings that must be matched from child to parent. */
 	TArray<FCurveEditorTreeTextFilterTerm, TInlineAllocator<4>> ChildToParentFilterTerms;
-};
\ No newline at end of file
+};
+
+#undef UE_API
diff --git a/Engine/Source/Editor/CurveEditor/Public/Tree/ICurveEditorTreeItem.h b/Engine/Source/Editor/CurveEditor/Public/Tree/ICurveEditorTreeItem.h
--- a/Engine/Source/Editor/CurveEditor/Public/Tree/ICurveEditorTreeItem.h
+++ b/Engine/Source/Editor/CurveEditor/Public/Tree/ICurveEditorTreeItem.h
@@ -8,6 +8,8 @@
 #include "Templates/UniquePtr.h"
 #include "UObject/NameTypes.h"
 
+#define UE_API CURVEEDITOR_API
+
 enum class ECurveEditorTreeFilterType : uint32;
 
 class FCurveEditor;
@@ -22,7 +24,7 @@ struct FCurveEditorTreeItemID;
 /**
  * Optional implementation interface for any tree item to be shown on the curve editor tree.
  */
-struct CURVEEDITOR_API ICurveEditorTreeItem
+struct ICurveEditorTreeItem
 {
 	virtual ~ICurveEditorTreeItem() {}
 
@@ -65,5 +67,7 @@ struct CURVEEDITOR_API ICurveEditorTreeItem
 		CURVEEDITOR_API FColumnNames();
 	};
 
-	static const FColumnNames ColumnNames;
+	static UE_API const FColumnNames ColumnNames;
 };
+
+#undef UE_API
diff --git a/Engine/Source/Editor/CurveEditor/Public/Tree/SCurveEditorTree.h b/Engine/Source/Editor/CurveEditor/Public/Tree/SCurveEditorTree.h
--- a/Engine/Source/Editor/CurveEditor/Public/Tree/SCurveEditorTree.h
+++ b/Engine/Source/Editor/CurveEditor/Public/Tree/SCurveEditorTree.h
@@ -19,6 +19,8 @@
 #include "Widgets/Views/STableViewBase.h"
 #include "Widgets/Views/STreeView.h"
 
+#define UE_API CURVEEDITOR_API
+
 class FCurveEditor;
 class ITableRow;
 class SHeaderRow;
@@ -26,7 +28,7 @@ class STableViewBase;
 struct FGeometry;
 struct FKeyEvent;
 
-class CURVEEDITOR_API SCurveEditorTree : public STreeView<FCurveEditorTreeItemID>
+class SCurveEditorTree : public STreeView<FCurveEditorTreeItemID>
 {
 
 public:
@@ -40,32 +42,32 @@ public:
 		SLATE_EVENT(FOnContextMenuOpening, OnContextMenuOpening)
 	SLATE_END_ARGS()
 
-	void Construct(const FArguments& InArgs, TSharedPtr<FCurveEditor> InCurveEditor);
+	UE_API void Construct(const FArguments& InArgs, TSharedPtr<FCurveEditor> InCurveEditor);
 
 	const TArray<FCurveEditorTreeItemID>& GetSourceItems() const { return RootItems; }
 
-	TSharedRef<ITableRow> GenerateRow(FCurveEditorTreeItemID ItemID, const TSharedRef<STableViewBase>& OwnerTable);
+	UE_API TSharedRef<ITableRow> GenerateRow(FCurveEditorTreeItemID ItemID, const TSharedRef<STableViewBase>& OwnerTable);
 	
 private:
 
-	virtual FReply OnKeyDown(const FGeometry& MyGeometry, const FKeyEvent& InKeyEvent) override;
+	UE_API virtual FReply OnKeyDown(const FGeometry& MyGeometry, const FKeyEvent& InKeyEvent) override;
 
 private:
 
 
-	void GetTreeItemChildren(FCurveEditorTreeItemID Parent, TArray<FCurveEditorTreeItemID>& OutChildren);
+	UE_API void GetTreeItemChildren(FCurveEditorTreeItemID Parent, TArray<FCurveEditorTreeItemID>& OutChildren);
 
-	void OnTreeSelectionChanged(FCurveEditorTreeItemID, ESelectInfo::Type);
+	UE_API void OnTreeSelectionChanged(FCurveEditorTreeItemID, ESelectInfo::Type);
 
-	void SetItemExpansionRecursive(FCurveEditorTreeItemID Model, bool bInExpansionState);
+	UE_API void SetItemExpansionRecursive(FCurveEditorTreeItemID Model, bool bInExpansionState);
 
-	void OnExpansionChanged(FCurveEditorTreeItemID Model, bool bInExpansionState);
+	UE_API void OnExpansionChanged(FCurveEditorTreeItemID Model, bool bInExpansionState);
 
-	void RefreshTree();
+	UE_API void RefreshTree();
 
-	void RefreshTreeWidgetSelection();
+	UE_API void RefreshTreeWidgetSelection();
 
-	void ToggleExpansionState(bool bRecursive);
+	UE_API void ToggleExpansionState(bool bRecursive);
 
 private:
 
@@ -83,4 +85,6 @@ private:
 	bool bUpdatingTreeWidgetSelection;
 	bool bUpdatingCurveEditorTreeSelection;
 
-};
\ No newline at end of file
+};
+
+#undef UE_API
diff --git a/Engine/Source/Editor/CurveEditor/Public/Tree/SCurveEditorTreeFilterStatusBar.h b/Engine/Source/Editor/CurveEditor/Public/Tree/SCurveEditorTreeFilterStatusBar.h
--- a/Engine/Source/Editor/CurveEditor/Public/Tree/SCurveEditorTreeFilterStatusBar.h
+++ b/Engine/Source/Editor/CurveEditor/Public/Tree/SCurveEditorTreeFilterStatusBar.h
@@ -7,39 +7,43 @@
 #include "Widgets/DeclarativeSyntaxSupport.h"
 #include "Widgets/SCompoundWidget.h"
 
+#define UE_API CURVEEDITOR_API
+
 class FCurveEditor;
 class STextBlock;
 
 /** Widget that shows the metrics for the current tree filter in the curve editor (in the form "Showing {0} of {1} items ({2} selected)") */
-class CURVEEDITOR_API SCurveEditorTreeFilterStatusBar : public SCompoundWidget
+class SCurveEditorTreeFilterStatusBar : public SCompoundWidget
 {
 public:
 	SLATE_BEGIN_ARGS(SCurveEditorTreeFilterStatusBar){}
 	SLATE_END_ARGS()
 
 	/** Construct the status bar */
-	void Construct(const FArguments& InArgs, TSharedPtr<FCurveEditor> InCurveEditor);
+	UE_API void Construct(const FArguments& InArgs, TSharedPtr<FCurveEditor> InCurveEditor);
 
-	virtual void Tick(const FGeometry& AllottedGeometry, const double InCurrentTime, const float InDeltaTime) override;
+	UE_API virtual void Tick(const FGeometry& AllottedGeometry, const double InCurrentTime, const float InDeltaTime) override;
 
-	void ShowStatusBar();
-	void HideStatusBar();
-	void FadeOutStatusBar();
+	UE_API void ShowStatusBar();
+	UE_API void HideStatusBar();
+	UE_API void FadeOutStatusBar();
 
 private:
 
 	/** Request the visibility of the clear hyperlink widget based on whether there's a filter active or not */
-	EVisibility GetVisibilityFromFilter() const;
+	UE_API EVisibility GetVisibilityFromFilter() const;
 
 	/** Update the filter text to represent the current filter states in the tree */
-	void UpdateText();
+	UE_API void UpdateText();
 
 	/** Clear all the filters in the tree - called as a result of the user clicking on the 'clear' hyperlink */
-	void ClearFilters();
+	UE_API void ClearFilters();
 
 private:
 
 	TWeakPtr<FCurveEditor> WeakCurveEditor;
 	TSharedPtr<STextBlock> TextBlock;
 	double OpacityThrobEndTime = 0;
-};
\ No newline at end of file
+};
+
+#undef UE_API
diff --git a/Engine/Source/Editor/CurveEditor/Public/Tree/SCurveEditorTreePin.h b/Engine/Source/Editor/CurveEditor/Public/Tree/SCurveEditorTreePin.h
--- a/Engine/Source/Editor/CurveEditor/Public/Tree/SCurveEditorTreePin.h
+++ b/Engine/Source/Editor/CurveEditor/Public/Tree/SCurveEditorTreePin.h
@@ -11,34 +11,38 @@
 #include "Widgets/DeclarativeSyntaxSupport.h"
 #include "Widgets/SCompoundWidget.h"
 
+#define UE_API CURVEEDITOR_API
+
 class FCurveEditor;
 class ITableRow;
 struct FSlateBrush;
 
-class CURVEEDITOR_API SCurveEditorTreePin : public SCompoundWidget
+class SCurveEditorTreePin : public SCompoundWidget
 {
 public:
 	SLATE_BEGIN_ARGS(SCurveEditorTreePin){}
 	SLATE_END_ARGS()
 
-	void Construct(const FArguments& InArgs, TWeakPtr<FCurveEditor> InCurveEditor, FCurveEditorTreeItemID InTreeItemID, const TSharedRef<ITableRow>& InTableRow);
+	UE_API void Construct(const FArguments& InArgs, TWeakPtr<FCurveEditor> InCurveEditor, FCurveEditorTreeItemID InTreeItemID, const TSharedRef<ITableRow>& InTableRow);
 
 private:
 
-	FReply TogglePinned();
-	const FSlateBrush* GetPinBrush() const;
+	UE_API FReply TogglePinned();
+	UE_API const FSlateBrush* GetPinBrush() const;
 
-	bool IsPinnedRecursive(FCurveEditorTreeItemID InTreeItem, FCurveEditor* CurveEditor) const;
+	UE_API bool IsPinnedRecursive(FCurveEditorTreeItemID InTreeItem, FCurveEditor* CurveEditor) const;
 
-	void PinRecursive(FCurveEditorTreeItemID InTreeItem, FCurveEditor* CurveEditor) const;
+	UE_API void PinRecursive(FCurveEditorTreeItemID InTreeItem, FCurveEditor* CurveEditor) const;
 
-	void UnpinRecursive(FCurveEditorTreeItemID InTreeItem, FCurveEditor* CurveEditor, TArray<FCurveEditorTreeItemID>& OutUnpinnedItems) const;
+	UE_API void UnpinRecursive(FCurveEditorTreeItemID InTreeItem, FCurveEditor* CurveEditor, TArray<FCurveEditorTreeItemID>& OutUnpinnedItems) const;
 
-	EVisibility GetPinVisibility() const;
+	UE_API EVisibility GetPinVisibility() const;
 
 private:
 
 	TWeakPtr<ITableRow> WeakTableRow;
 	TWeakPtr<FCurveEditor> WeakCurveEditor;
 	FCurveEditorTreeItemID TreeItemID;
-};
\ No newline at end of file
+};
+
+#undef UE_API
diff --git a/Engine/Source/Editor/CurveEditor/Public/Tree/SCurveEditorTreeSelect.h b/Engine/Source/Editor/CurveEditor/Public/Tree/SCurveEditorTreeSelect.h
--- a/Engine/Source/Editor/CurveEditor/Public/Tree/SCurveEditorTreeSelect.h
+++ b/Engine/Source/Editor/CurveEditor/Public/Tree/SCurveEditorTreeSelect.h
@@ -10,29 +10,33 @@
 #include "Widgets/DeclarativeSyntaxSupport.h"
 #include "Widgets/SCompoundWidget.h"
 
+#define UE_API CURVEEDITOR_API
+
 class FCurveEditor;
 class ITableRow;
 struct FSlateBrush;
 
-class CURVEEDITOR_API SCurveEditorTreeSelect : public SCompoundWidget
+class SCurveEditorTreeSelect : public SCompoundWidget
 {
 public:
 	SLATE_BEGIN_ARGS(SCurveEditorTreeSelect){}
 	SLATE_END_ARGS()
 
-	void Construct(const FArguments& InArgs, TWeakPtr<FCurveEditor> InCurveEditor, FCurveEditorTreeItemID InTreeItemID, const TSharedRef<ITableRow>& InTableRow);
+	UE_API void Construct(const FArguments& InArgs, TWeakPtr<FCurveEditor> InCurveEditor, FCurveEditorTreeItemID InTreeItemID, const TSharedRef<ITableRow>& InTableRow);
 
 private:
 
-	FReply SelectAll();
+	UE_API FReply SelectAll();
 
-	const FSlateBrush* GetSelectBrush() const;
+	UE_API const FSlateBrush* GetSelectBrush() const;
 
-	EVisibility GetSelectVisibility() const;
+	UE_API EVisibility GetSelectVisibility() const;
 
 private:
 
 	TWeakPtr<ITableRow> WeakTableRow;
 	TWeakPtr<FCurveEditor> WeakCurveEditor;
 	FCurveEditorTreeItemID TreeItemID;
-};
\ No newline at end of file
+};
+
+#undef UE_API
diff --git a/Engine/Source/Editor/CurveEditor/Public/Tree/SCurveEditorTreeTextFilter.h b/Engine/Source/Editor/CurveEditor/Public/Tree/SCurveEditorTreeTextFilter.h
--- a/Engine/Source/Editor/CurveEditor/Public/Tree/SCurveEditorTreeTextFilter.h
+++ b/Engine/Source/Editor/CurveEditor/Public/Tree/SCurveEditorTreeTextFilter.h
@@ -7,6 +7,8 @@
 #include "Widgets/DeclarativeSyntaxSupport.h"
 #include "Widgets/SCompoundWidget.h"
 
+#define UE_API CURVEEDITOR_API
+
 class FCurveEditor;
 class FText;
 class SSearchBox;
@@ -14,26 +16,28 @@ struct FCurveEditorTreeTextFilter;
 struct FFocusEvent;
 struct FGeometry;
 
-class CURVEEDITOR_API SCurveEditorTreeTextFilter : public SCompoundWidget
+class SCurveEditorTreeTextFilter : public SCompoundWidget
 {
 public:
 
 	SLATE_BEGIN_ARGS(SCurveEditorTreeTextFilter){}
 	SLATE_END_ARGS()
 
-	void Construct(const FArguments& InArgs, TSharedPtr<FCurveEditor> CurveEditor);
+	UE_API void Construct(const FArguments& InArgs, TSharedPtr<FCurveEditor> CurveEditor);
 	
 	// SWidget Interface
-	virtual FReply OnFocusReceived(const FGeometry& MyGeometry, const FFocusEvent& InFocusEvent) override;
+	UE_API virtual FReply OnFocusReceived(const FGeometry& MyGeometry, const FFocusEvent& InFocusEvent) override;
 	virtual bool SupportsKeyboardFocus() const override { return true; }
 	// ~SWidget Interface
 private:
 
-	void CreateSearchBox();
-	void OnTreeFilterListChanged();
-	void OnFilterTextChanged(const FText& FilterText);
+	UE_API void CreateSearchBox();
+	UE_API void OnTreeFilterListChanged();
+	UE_API void OnFilterTextChanged(const FText& FilterText);
 
 	TSharedPtr<SSearchBox> SearchBox;
 	TSharedPtr<FCurveEditorTreeTextFilter> Filter;
 	TWeakPtr<FCurveEditor> WeakCurveEditor;
-};
\ No newline at end of file
+};
+
+#undef UE_API
diff --git a/Engine/Source/Editor/CurveEditor/Public/Views/SCurveEditorViewAbsolute.h b/Engine/Source/Editor/CurveEditor/Public/Views/SCurveEditorViewAbsolute.h
--- a/Engine/Source/Editor/CurveEditor/Public/Views/SCurveEditorViewAbsolute.h
+++ b/Engine/Source/Editor/CurveEditor/Public/Views/SCurveEditorViewAbsolute.h
@@ -5,15 +5,19 @@
 #include "Templates/SharedPointer.h"
 #include "Views/SInteractiveCurveEditorView.h"
 
+#define UE_API CURVEEDITOR_API
+
 class FCurveEditor;
 
 /**
  * A Normalized curve view supporting one or more curves with their own screen transform that normalizes the vertical curve range to [-1,1]
  */
-class CURVEEDITOR_API SCurveEditorViewAbsolute : public SInteractiveCurveEditorView
+class SCurveEditorViewAbsolute : public SInteractiveCurveEditorView
 {
 public:
 
-	void Construct(const FArguments& InArgs, TWeakPtr<FCurveEditor> InCurveEditor);
+	UE_API void Construct(const FArguments& InArgs, TWeakPtr<FCurveEditor> InCurveEditor);
 
 };
+
+#undef UE_API
diff --git a/Engine/Source/Editor/CurveEditor/Public/Views/SCurveEditorViewStacked.h b/Engine/Source/Editor/CurveEditor/Public/Views/SCurveEditorViewStacked.h
--- a/Engine/Source/Editor/CurveEditor/Public/Views/SCurveEditorViewStacked.h
+++ b/Engine/Source/Editor/CurveEditor/Public/Views/SCurveEditorViewStacked.h
@@ -9,6 +9,8 @@
 #include "Templates/SharedPointer.h"
 #include "Views/SInteractiveCurveEditorView.h"
 
+#define UE_API CURVEEDITOR_API
+
 class FCurveEditor;
 class FPaintArgs;
 class FSlateRect;
@@ -20,28 +22,28 @@ struct FGeometry;
 /**
  * A Normalized curve view supporting one or more curves with their own screen transform that normalizes the vertical curve range to [-1,1]
  */
-class CURVEEDITOR_API SCurveEditorViewStacked : public SInteractiveCurveEditorView
+class SCurveEditorViewStacked : public SInteractiveCurveEditorView
 {
 public:
 
-	void Construct(const FArguments& InArgs, TWeakPtr<FCurveEditor> InCurveEditor);
+	UE_API void Construct(const FArguments& InArgs, TWeakPtr<FCurveEditor> InCurveEditor);
 	
 	/** Tools should use vertical snapping since grid lines to snap to will usually be visible */
 	virtual bool IsValueSnapEnabled() const override { return true; }
 
-	virtual void UpdateViewToTransformCurves(double InputMin, double InputMax) override;
+	UE_API virtual void UpdateViewToTransformCurves(double InputMin, double InputMax) override;
 
 protected:
 
-	virtual void Tick(const FGeometry& AllottedGeometry, const double InCurrentTime, const float InDeltaTime) override;
-	virtual FVector2D ComputeDesiredSize(float LayoutScaleMultiplier) const override;
+	UE_API virtual void Tick(const FGeometry& AllottedGeometry, const double InCurrentTime, const float InDeltaTime) override;
+	UE_API virtual FVector2D ComputeDesiredSize(float LayoutScaleMultiplier) const override;
 
-	virtual void PaintView(const FPaintArgs& Args, const FGeometry& AllottedGeometry, const FSlateRect& MyCullingRect, FSlateWindowElementList& OutDrawElements, int32 BaseLayerId, const FWidgetStyle& InWidgetStyle, bool bParentEnabled) const override;
-	virtual void GetGridLinesY(TSharedRef<const FCurveEditor> CurveEditor, TArray<float>& MajorGridLines, TArray<float>& MinorGridLines, TArray<FText>* MajorGridLabels = nullptr) const override;
+	UE_API virtual void PaintView(const FPaintArgs& Args, const FGeometry& AllottedGeometry, const FSlateRect& MyCullingRect, FSlateWindowElementList& OutDrawElements, int32 BaseLayerId, const FWidgetStyle& InWidgetStyle, bool bParentEnabled) const override;
+	UE_API virtual void GetGridLinesY(TSharedRef<const FCurveEditor> CurveEditor, TArray<float>& MajorGridLines, TArray<float>& MinorGridLines, TArray<FText>* MajorGridLabels = nullptr) const override;
 
-	virtual void DrawViewGrids(const FGeometry& AllottedGeometry, const FSlateRect& MyCullingRect, FSlateWindowElementList& OutDrawElements, int32 BaseLayerId, ESlateDrawEffect DrawEffects) const;
-	virtual void DrawLabels(const FGeometry& AllottedGeometry, const FSlateRect& MyCullingRect, FSlateWindowElementList& OutDrawElements, int32 BaseLayerId, ESlateDrawEffect DrawEffects) const;
-	virtual void DrawBufferedCurves(const FGeometry& AllottedGeometry, const FSlateRect& MyCullingRect, FSlateWindowElementList& OutDrawElements, int32 BaseLayerId, ESlateDrawEffect DrawEffects) const;
+	UE_API virtual void DrawViewGrids(const FGeometry& AllottedGeometry, const FSlateRect& MyCullingRect, FSlateWindowElementList& OutDrawElements, int32 BaseLayerId, ESlateDrawEffect DrawEffects) const;
+	UE_API virtual void DrawLabels(const FGeometry& AllottedGeometry, const FSlateRect& MyCullingRect, FSlateWindowElementList& OutDrawElements, int32 BaseLayerId, ESlateDrawEffect DrawEffects) const;
+	UE_API virtual void DrawBufferedCurves(const FGeometry& AllottedGeometry, const FSlateRect& MyCullingRect, FSlateWindowElementList& OutDrawElements, int32 BaseLayerId, ESlateDrawEffect DrawEffects) const;
 
 	/** Stacked height per curve */
 	float StackedHeight;
@@ -49,3 +51,5 @@ protected:
 	/** Stacked padding per curve */
 	float StackedPadding;
 };
+
+#undef UE_API
diff --git a/Engine/Source/Editor/CurveEditor/Public/Views/SInteractiveCurveEditorView.h b/Engine/Source/Editor/CurveEditor/Public/Views/SInteractiveCurveEditorView.h
--- a/Engine/Source/Editor/CurveEditor/Public/Views/SInteractiveCurveEditorView.h
+++ b/Engine/Source/Editor/CurveEditor/Public/Views/SInteractiveCurveEditorView.h
@@ -26,6 +26,8 @@
 #include "Widgets/SCompoundWidget.h"
 #include "CurveEditorSettings.h"
 
+#define UE_API CURVEEDITOR_API
+
 class FCurveEditor;
 class FCurveModel;
 class FMenuBuilder;
@@ -80,7 +82,7 @@ namespace CurveViewConstants
 
 /**
  */
-class CURVEEDITOR_API SInteractiveCurveEditorView : public SCurveEditorView
+class SInteractiveCurveEditorView : public SCurveEditorView
 {
 public:
 
@@ -100,71 +102,71 @@ public:
 
 	SLATE_END_ARGS()
 
-	void Construct(const FArguments& InArgs, TWeakPtr<FCurveEditor> InCurveEditor);
+	UE_API void Construct(const FArguments& InArgs, TWeakPtr<FCurveEditor> InCurveEditor);
 
-	virtual void GetGridLinesX(TSharedRef<const FCurveEditor> CurveEditor, TArray<float>& MajorGridLines, TArray<float>& MinorGridLines, TArray<FText>* MajorGridLabels = nullptr) const override;
-	virtual void GetGridLinesY(TSharedRef<const FCurveEditor> CurveEditor, TArray<float>& MajorGridLines, TArray<float>& MinorGridLines, TArray<FText>* MajorGridLabels = nullptr) const override;
+	UE_API virtual void GetGridLinesX(TSharedRef<const FCurveEditor> CurveEditor, TArray<float>& MajorGridLines, TArray<float>& MinorGridLines, TArray<FText>* MajorGridLabels = nullptr) const override;
+	UE_API virtual void GetGridLinesY(TSharedRef<const FCurveEditor> CurveEditor, TArray<float>& MajorGridLines, TArray<float>& MinorGridLines, TArray<FText>* MajorGridLabels = nullptr) const override;
 
-	virtual void BuildContextMenu(FMenuBuilder& MenuBuilder, TOptional<FCurvePointHandle> ClickedPoint, TOptional<FCurveModelID> HoveredCurveID);
+	UE_API virtual void BuildContextMenu(FMenuBuilder& MenuBuilder, TOptional<FCurvePointHandle> ClickedPoint, TOptional<FCurveModelID> HoveredCurveID);
 
 protected:
 
 	// ~SCurveEditorView Interface
-	virtual bool GetPointsWithinWidgetRange(const FSlateRect& WidgetRectangle, TArray<FCurvePointHandle>* OutPoints) const override;
-	virtual bool GetCurveWithinWidgetRange(const FSlateRect& WidgetRectangle, TArray<FCurvePointHandle>* OutPoints) const override;
-	virtual TOptional<FCurveModelID> GetHoveredCurve() const override;
+	UE_API virtual bool GetPointsWithinWidgetRange(const FSlateRect& WidgetRectangle, TArray<FCurvePointHandle>* OutPoints) const override;
+	UE_API virtual bool GetCurveWithinWidgetRange(const FSlateRect& WidgetRectangle, TArray<FCurvePointHandle>* OutPoints) const override;
+	UE_API virtual TOptional<FCurveModelID> GetHoveredCurve() const override;
 
-	virtual FText FormatToolTipCurveName(const FCurveModel& CurveModel) const;
-	virtual FText FormatToolTipTime(const FCurveModel& CurveModel, double EvaluatedTime) const;
-	virtual FText FormatToolTipValue(const FCurveModel& CurveModel, double EvaluatedValue) const;
+	UE_API virtual FText FormatToolTipCurveName(const FCurveModel& CurveModel) const;
+	UE_API virtual FText FormatToolTipTime(const FCurveModel& CurveModel, double EvaluatedTime) const;
+	UE_API virtual FText FormatToolTipValue(const FCurveModel& CurveModel, double EvaluatedValue) const;
 
-	virtual void PaintView(const FPaintArgs& Args, const FGeometry& AllottedGeometry, const FSlateRect& MyCullingRect, FSlateWindowElementList& OutDrawElements, int32 BaseLayerId, const FWidgetStyle& InWidgetStyle, bool bParentEnabled) const;
+	UE_API virtual void PaintView(const FPaintArgs& Args, const FGeometry& AllottedGeometry, const FSlateRect& MyCullingRect, FSlateWindowElementList& OutDrawElements, int32 BaseLayerId, const FWidgetStyle& InWidgetStyle, bool bParentEnabled) const;
 
 protected:
 
 	// SWidget Interface
-	virtual void Tick(const FGeometry& AllottedGeometry, const double InCurrentTime, const float InDeltaTime) override;
-	virtual int32 OnPaint(const FPaintArgs& Args, const FGeometry& AllottedGeometry, const FSlateRect& MyCullingRect, FSlateWindowElementList& OutDrawElements, int32 BaseLayerId, const FWidgetStyle& InWidgetStyle, bool bParentEnabled) const override;
-	virtual void OnMouseEnter(const FGeometry& MyGeometry, const FPointerEvent& MouseEvent) override;
-	virtual void OnMouseLeave(const FPointerEvent& MouseEvent) override;
-	virtual FReply OnMouseMove(const FGeometry& MyGeometry, const FPointerEvent& MouseEvent) override;
-	virtual void OnFinishedPointerInput() override;
-	virtual FReply OnMouseWheel(const FGeometry& MyGeometry, const FPointerEvent& MouseEvent) override;
-	virtual FReply OnMouseButtonUp(const FGeometry& MyGeometry, const FPointerEvent& MouseEvent) override;
-	virtual FReply OnMouseButtonDown(const FGeometry& MyGeometry, const FPointerEvent& MouseEvent) override;
+	UE_API virtual void Tick(const FGeometry& AllottedGeometry, const double InCurrentTime, const float InDeltaTime) override;
+	UE_API virtual int32 OnPaint(const FPaintArgs& Args, const FGeometry& AllottedGeometry, const FSlateRect& MyCullingRect, FSlateWindowElementList& OutDrawElements, int32 BaseLayerId, const FWidgetStyle& InWidgetStyle, bool bParentEnabled) const override;
+	UE_API virtual void OnMouseEnter(const FGeometry& MyGeometry, const FPointerEvent& MouseEvent) override;
+	UE_API virtual void OnMouseLeave(const FPointerEvent& MouseEvent) override;
+	UE_API virtual FReply OnMouseMove(const FGeometry& MyGeometry, const FPointerEvent& MouseEvent) override;
+	UE_API virtual void OnFinishedPointerInput() override;
+	UE_API virtual FReply OnMouseWheel(const FGeometry& MyGeometry, const FPointerEvent& MouseEvent) override;
+	UE_API virtual FReply OnMouseButtonUp(const FGeometry& MyGeometry, const FPointerEvent& MouseEvent) override;
+	UE_API virtual FReply OnMouseButtonDown(const FGeometry& MyGeometry, const FPointerEvent& MouseEvent) override;
 	// ~SWidget Interface
 
 protected:
 
-	void DrawBackground(const FGeometry& AllottedGeometry, FSlateWindowElementList& OutDrawElements, int32 BaseLayerId, ESlateDrawEffect DrawEffects) const;
-	void DrawGridLines(TSharedRef<FCurveEditor> CurveEditor, const FGeometry& AllottedGeometry, FSlateWindowElementList& OutDrawElements, int32 BaseLayerId, ESlateDrawEffect DrawEffects) const;
-	void DrawCurves(TSharedRef<FCurveEditor> CurveEditor, const FGeometry& AllottedGeometry, const FSlateRect& MyCullingRect, FSlateWindowElementList& OutDrawElements, int32 BaseLayerId, const FWidgetStyle& InWidgetStyle, ESlateDrawEffect DrawEffects) const;
-	void DrawBufferedCurves(TSharedRef<FCurveEditor> CurveEditor, const FGeometry& AllottedGeometry, const FSlateRect& MyCullingRect, FSlateWindowElementList& OutDrawElements, int32 BaseLayerId, const FWidgetStyle& InWidgetStyle, ESlateDrawEffect DrawEffects) const;
-	void DrawValueIndicatorLines(TSharedRef<FCurveEditor> InCurveEditor, const FGeometry& InAllottedGeometry, FSlateWindowElementList& OutDrawElements, int32 InBaseLayerId) const;
+	UE_API void DrawBackground(const FGeometry& AllottedGeometry, FSlateWindowElementList& OutDrawElements, int32 BaseLayerId, ESlateDrawEffect DrawEffects) const;
+	UE_API void DrawGridLines(TSharedRef<FCurveEditor> CurveEditor, const FGeometry& AllottedGeometry, FSlateWindowElementList& OutDrawElements, int32 BaseLayerId, ESlateDrawEffect DrawEffects) const;
+	UE_API void DrawCurves(TSharedRef<FCurveEditor> CurveEditor, const FGeometry& AllottedGeometry, const FSlateRect& MyCullingRect, FSlateWindowElementList& OutDrawElements, int32 BaseLayerId, const FWidgetStyle& InWidgetStyle, ESlateDrawEffect DrawEffects) const;
+	UE_API void DrawBufferedCurves(TSharedRef<FCurveEditor> CurveEditor, const FGeometry& AllottedGeometry, const FSlateRect& MyCullingRect, FSlateWindowElementList& OutDrawElements, int32 BaseLayerId, const FWidgetStyle& InWidgetStyle, ESlateDrawEffect DrawEffects) const;
+	UE_API void DrawValueIndicatorLines(TSharedRef<FCurveEditor> InCurveEditor, const FGeometry& InAllottedGeometry, FSlateWindowElementList& OutDrawElements, int32 InBaseLayerId) const;
 	
 	/** Updates the keys that will be drawn with value indicator lines. */
-	void UpdatedKeysWithValueIndicatorLines(const FCurveEditor& InCurveEditor) const;
+	UE_API void UpdatedKeysWithValueIndicatorLines(const FCurveEditor& InCurveEditor) const;
 	/** Picks which keys will have value indicator lines drawn for them. Returns true if the selection. */
-	bool PickPointsToPlaceValueIndicatorLinesOn(const FCurveModel& InCurveModel, const FKeyHandleSet& InUserSelectedKeys, FKeyHandle& OutMinKey, FKeyHandle& OutMaxKey) const;
+	UE_API bool PickPointsToPlaceValueIndicatorLinesOn(const FCurveModel& InCurveModel, const FKeyHandleSet& InUserSelectedKeys, FKeyHandle& OutMinKey, FKeyHandle& OutMaxKey) const;
 	
-	FSlateColor GetCurveCaptionColor() const;
-	FText GetCurveCaption() const;
+	UE_API FSlateColor GetCurveCaptionColor() const;
+	UE_API FText GetCurveCaption() const;
 
 private:
 	
-	void HandleDirectKeySelectionByMouse(TSharedPtr<FCurveEditor> CurveEditor, const FPointerEvent& MouseEvent, TOptional<FCurvePointHandle> MouseDownPoint);
+	UE_API void HandleDirectKeySelectionByMouse(TSharedPtr<FCurveEditor> CurveEditor, const FPointerEvent& MouseEvent, TOptional<FCurvePointHandle> MouseDownPoint);
 
-	void CreateContextMenu(const FGeometry& MyGeometry, const FPointerEvent& MouseEvent);
+	UE_API void CreateContextMenu(const FGeometry& MyGeometry, const FPointerEvent& MouseEvent);
 
 	/** Updates our distance to all of the curves we represent. */
-	void UpdateCurveProximities(FVector2D MousePixel);
+	UE_API void UpdateCurveProximities(FVector2D MousePixel);
 
-	TOptional<FCurvePointHandle> HitPoint(FVector2D MousePixel) const;
+	UE_API TOptional<FCurvePointHandle> HitPoint(FVector2D MousePixel) const;
 
-	bool IsToolTipEnabled() const;
-	FText GetToolTipCurveName() const;
-	FText GetToolTipTimeText() const;
-	FText GetToolTipValueText() const;
+	UE_API bool IsToolTipEnabled() const;
+	UE_API FText GetToolTipCurveName() const;
+	UE_API FText GetToolTipTimeText() const;
+	UE_API FText GetToolTipValueText() const;
 
 	/**
 	 * Returns the proper tangent value so we can keep the curve remain the original shape
@@ -175,32 +177,32 @@ private:
 	 * @param DeltaTime		Negative to get the left tangent, positive for right. Remember to use FMath::Abs() when needed
 	 * @return				The tangent value relatives to the DeltaTime upon mouse click's position
 	 */
-	double GetTangentValue(const double InTime, const double InValue, FCurveModel* CurveToAddTo, double DeltaTime) const;
+	UE_API double GetTangentValue(const double InTime, const double InValue, FCurveModel* CurveToAddTo, double DeltaTime) const;
 
 	/*~ Command binding callbacks */
-	void AddKeyAtScrubTime(TSet<FCurveModelID> ForCurves);
-	void AddKeyAtMousePosition(TSet<FCurveModelID> ForCurves);
-	void AddKeyAtTime(const TSet<FCurveModelID>& ToCurves, double InTime);
-	void PasteKeys(TSet<FCurveModelID> ToCurves);
+	UE_API void AddKeyAtScrubTime(TSet<FCurveModelID> ForCurves);
+	UE_API void AddKeyAtMousePosition(TSet<FCurveModelID> ForCurves);
+	UE_API void AddKeyAtTime(const TSet<FCurveModelID>& ToCurves, double InTime);
+	UE_API void PasteKeys(TSet<FCurveModelID> ToCurves);
 
-	void OnCurveEditorToolChanged(FCurveEditorToolID InToolId);
-	void OnShowValueIndicatorsChanged();
+	UE_API void OnCurveEditorToolChanged(FCurveEditorToolID InToolId);
+	UE_API void OnShowValueIndicatorsChanged();
 
 	/**
 	 * Rebind contextual command mappings that rely on the mouse position
 	 */
-	void RebindContextualActions(FVector2D InMousePosition);
+	UE_API void RebindContextualActions(FVector2D InMousePosition);
 
 	/** Copy the curves from this view and set them as the Curve Editor's buffered curve support. */
-	void BufferCurves();
+	UE_API void BufferCurves();
 	/** Attempt to apply the previously buffered curves to the currently selected curves. */
-	void ApplyBufferCurves(const bool bSwapBufferCurves);
+	UE_API void ApplyBufferCurves(const bool bSwapBufferCurves);
 	/** Check if it's legal to buffer any of our selected curves. */
-	bool CanBufferedCurves() const;
+	UE_API bool CanBufferedCurves() const;
 	/** Check if it's legal to apply any of the buffered curves to our currently selected curves. */
-	bool CanApplyBufferedCurves() const;
+	UE_API bool CanApplyBufferedCurves() const;
 	/** Returns interpolation mode and tangent mode based on neighbours or default curve editor if no neighbours . */
-	FKeyAttributes GetDefaultKeyAttributesForCurveTime(const FCurveEditor& CurveEditor, const FCurveModel& CurveModel, double EvalTime) const;
+	UE_API FKeyAttributes GetDefaultKeyAttributesForCurveTime(const FCurveEditor& CurveEditor, const FCurveModel& CurveModel, double EvalTime) const;
 
 protected:
 
@@ -269,4 +271,6 @@ private:
 	 * This is called in OnPaint, which is const. Mutable is a workaround.
 	 */
 	mutable TOptional<FValueIndicatorLineData> ValueIndicatorLineDrawData;
-};
\ No newline at end of file
+};
+
+#undef UE_API
diff --git a/Engine/Source/Editor/DataHierarchyEditor/Internal/DataHierarchyEditorStyle.h b/Engine/Source/Editor/DataHierarchyEditor/Internal/DataHierarchyEditorStyle.h
--- a/Engine/Source/Editor/DataHierarchyEditor/Internal/DataHierarchyEditorStyle.h
+++ b/Engine/Source/Editor/DataHierarchyEditor/Internal/DataHierarchyEditorStyle.h
@@ -4,26 +4,30 @@
 
 #include "Styling/SlateStyle.h"
 
+#define UE_API DATAHIERARCHYEDITOR_API
+
 /** Manages the style which provides resources for Hierarchy Editor widgets. */
-class DATAHIERARCHYEDITOR_API FDataHierarchyEditorStyle : public FSlateStyleSet
+class FDataHierarchyEditorStyle : public FSlateStyleSet
 {
 public:
-	static void Register();
-	static void Unregister();
-	static void Shutdown();
+	static UE_API void Register();
+	static UE_API void Unregister();
+	static UE_API void Shutdown();
 
 	/** reloads textures used by slate renderer */
-	static void ReloadTextures();
+	static UE_API void ReloadTextures();
 
 	/** @return The Slate style set for Hierarchy Editor widgets */
-	static const FDataHierarchyEditorStyle& Get();
+	static UE_API const FDataHierarchyEditorStyle& Get();
 
-	static void ReinitializeStyle();
+	static UE_API void ReinitializeStyle();
 
 private:	
-	FDataHierarchyEditorStyle();
+	UE_API FDataHierarchyEditorStyle();
 
-	void InitDataHierarchyEditor();
+	UE_API void InitDataHierarchyEditor();
 
-	static TSharedPtr<FDataHierarchyEditorStyle> DataHierarchyEditorStyle;
+	static UE_API TSharedPtr<FDataHierarchyEditorStyle> DataHierarchyEditorStyle;
 };
+
+#undef UE_API
diff --git a/Engine/Source/Editor/DataHierarchyEditor/Internal/Widgets/SDataHierarchyEditor.h b/Engine/Source/Editor/DataHierarchyEditor/Internal/Widgets/SDataHierarchyEditor.h
--- a/Engine/Source/Editor/DataHierarchyEditor/Internal/Widgets/SDataHierarchyEditor.h
+++ b/Engine/Source/Editor/DataHierarchyEditor/Internal/Widgets/SDataHierarchyEditor.h
@@ -17,17 +17,19 @@
 #include "Misc/NotifyHook.h"
 #include "UObject/GCObject.h"
 
+#define UE_API DATAHIERARCHYEDITOR_API
+
 class UDataHierarchyViewModelBase;
 class UHierarchySection;
 struct FHierarchyElementViewModel;
 struct FHierarchySectionViewModel;
 
-class DATAHIERARCHYEDITOR_API SHierarchySection : public SCompoundWidget
+class SHierarchySection : public SCompoundWidget
 {
 	DECLARE_DELEGATE_OneParam(FOnSectionActivated, TSharedPtr<FHierarchySectionViewModel> SectionViewModel)
 
-	static const FSlateBrush LeftDropBrush;
-	static const FSlateBrush RightDropBrush;
+	static UE_API const FSlateBrush LeftDropBrush;
+	static UE_API const FSlateBrush RightDropBrush;
 	
 	SLATE_BEGIN_ARGS(SHierarchySection)
 		{}
@@ -36,51 +38,51 @@ class DATAHIERARCHYEDITOR_API SHierarchySection : public SCompoundWidget
 	SLATE_END_ARGS()
 
 public:
-	void Construct(const FArguments& InArgs, TSharedPtr<struct FHierarchySectionViewModel> InSection);
-	virtual ~SHierarchySection() override;
+	UE_API void Construct(const FArguments& InArgs, TSharedPtr<struct FHierarchySectionViewModel> InSection);
+	UE_API virtual ~SHierarchySection() override;
 
-	void TryEnterEditingMode() const;
+	UE_API void TryEnterEditingMode() const;
 
-	TSharedPtr<struct FHierarchySectionViewModel> GetSectionViewModel();
+	UE_API TSharedPtr<struct FHierarchySectionViewModel> GetSectionViewModel();
 private:
-	virtual int32 OnPaint(const FPaintArgs& Args, const FGeometry& AllottedGeometry, const FSlateRect& MyCullingRect, FSlateWindowElementList& OutDrawElements, int32 LayerId, const FWidgetStyle& InWidgetStyle, bool bParentEnabled) const override;
-	int32 PaintDropIndicator(const FPaintArgs& Args, const FGeometry& Geometry, FSlateRect SlateRect, FSlateWindowElementList& OutDrawElements, int32 INT32, const FWidgetStyle& WidgetStyle, bool bParentEnabled) const;
-	int32 OnPaintDropIndicator(EItemDropZone ItemDropZone, const FPaintArgs& Args, const FGeometry& Geometry, FSlateRect SlateRect, FSlateWindowElementList& OutDrawElements, int32 LayerId, const FWidgetStyle& WidgetStyle, bool bParentEnabled) const;
+	UE_API virtual int32 OnPaint(const FPaintArgs& Args, const FGeometry& AllottedGeometry, const FSlateRect& MyCullingRect, FSlateWindowElementList& OutDrawElements, int32 LayerId, const FWidgetStyle& InWidgetStyle, bool bParentEnabled) const override;
+	UE_API int32 PaintDropIndicator(const FPaintArgs& Args, const FGeometry& Geometry, FSlateRect SlateRect, FSlateWindowElementList& OutDrawElements, int32 INT32, const FWidgetStyle& WidgetStyle, bool bParentEnabled) const;
+	UE_API int32 OnPaintDropIndicator(EItemDropZone ItemDropZone, const FPaintArgs& Args, const FGeometry& Geometry, FSlateRect SlateRect, FSlateWindowElementList& OutDrawElements, int32 LayerId, const FWidgetStyle& WidgetStyle, bool bParentEnabled) const;
 
-	virtual FReply OnMouseButtonDown(const FGeometry& MyGeometry, const FPointerEvent& MouseEvent) override;
-	virtual FReply OnMouseButtonUp(const FGeometry& MyGeometry, const FPointerEvent& MouseEvent) override;
+	UE_API virtual FReply OnMouseButtonDown(const FGeometry& MyGeometry, const FPointerEvent& MouseEvent) override;
+	UE_API virtual FReply OnMouseButtonUp(const FGeometry& MyGeometry, const FPointerEvent& MouseEvent) override;
 
-	virtual FReply OnDrop(const FGeometry& MyGeometry, const FDragDropEvent& DragDropEvent) override;
+	UE_API virtual FReply OnDrop(const FGeometry& MyGeometry, const FDragDropEvent& DragDropEvent) override;
 	
-	bool OnCanAcceptDrop(TSharedPtr<FDragDropOperation> DragDropOperation, EItemDropZone ItemDropZone) const;
-	FReply OnDroppedOn(const FGeometry&, const FDragDropEvent& DragDropEvent, EItemDropZone DropZone) const;
+	UE_API bool OnCanAcceptDrop(TSharedPtr<FDragDropOperation> DragDropOperation, EItemDropZone ItemDropZone) const;
+	UE_API FReply OnDroppedOn(const FGeometry&, const FDragDropEvent& DragDropEvent, EItemDropZone DropZone) const;
 	
-	virtual void OnDragEnter(const FGeometry& MyGeometry, const FDragDropEvent& DragDropEvent) override;
-	virtual FReply OnDragDetected(const FGeometry& MyGeometry, const FPointerEvent& MouseEvent) override;
-	virtual void OnDragLeave(const FDragDropEvent& DragDropEvent) override;
-	virtual FReply OnDragOver(const FGeometry& MyGeometry, const FDragDropEvent& DragDropEvent) override;
+	UE_API virtual void OnDragEnter(const FGeometry& MyGeometry, const FDragDropEvent& DragDropEvent) override;
+	UE_API virtual FReply OnDragDetected(const FGeometry& MyGeometry, const FPointerEvent& MouseEvent) override;
+	UE_API virtual void OnDragLeave(const FDragDropEvent& DragDropEvent) override;
+	UE_API virtual FReply OnDragOver(const FGeometry& MyGeometry, const FDragDropEvent& DragDropEvent) override;
 
-	virtual FReply OnKeyDown(const FGeometry& MyGeometry, const FKeyEvent& InKeyEvent) override;
+	UE_API virtual FReply OnKeyDown(const FGeometry& MyGeometry, const FKeyEvent& InKeyEvent) override;
 	
-	virtual void OnMouseLeave(const FPointerEvent& MouseEvent) override;
+	UE_API virtual void OnMouseLeave(const FPointerEvent& MouseEvent) override;
 
-	UHierarchySection* TryGetSectionData() const;
+	UE_API UHierarchySection* TryGetSectionData() const;
 
-	FText GetText() const;
-	FText GetTooltipText() const;
-	void OnRenameSection(const FText& Text, ETextCommit::Type CommitType) const;
-	bool OnVerifySectionRename(const FText& NewName, FText& OutTooltip) const;
+	UE_API FText GetText() const;
+	UE_API FText GetTooltipText() const;
+	UE_API void OnRenameSection(const FText& Text, ETextCommit::Type CommitType) const;
+	UE_API bool OnVerifySectionRename(const FText& NewName, FText& OutTooltip) const;
 
-	bool IsSectionSelected() const;
-	bool IsSectionReadOnly() const;
-	ECheckBoxState GetSectionCheckState() const;
-	void OnSectionCheckChanged(ECheckBoxState NewState);
-	EActiveTimerReturnType ActivateSectionIfDragging(double CurrentTime, float DeltaTime) const;
+	UE_API bool IsSectionSelected() const;
+	UE_API bool IsSectionReadOnly() const;
+	UE_API ECheckBoxState GetSectionCheckState() const;
+	UE_API void OnSectionCheckChanged(ECheckBoxState NewState);
+	UE_API EActiveTimerReturnType ActivateSectionIfDragging(double CurrentTime, float DeltaTime) const;
 
-	const FSlateBrush* GetDropIndicatorBrush(EItemDropZone ItemDropZone) const;
+	UE_API const FSlateBrush* GetDropIndicatorBrush(EItemDropZone ItemDropZone) const;
 
 	/** @return the zone (above, onto, below) based on where the user is hovering over within the row */
-	EItemDropZone ZoneFromPointerPosition(UE::Slate::FDeprecateVector2DParameter LocalPointerPos, UE::Slate::FDeprecateVector2DParameter LocalSize);
+	UE_API EItemDropZone ZoneFromPointerPosition(UE::Slate::FDeprecateVector2DParameter LocalPointerPos, UE::Slate::FDeprecateVector2DParameter LocalSize);
 
 private:
 	TSharedPtr<SMenuAnchor> MenuAnchor;
@@ -97,7 +99,7 @@ private:
 	mutable TOptional<EItemDropZone> CurrentItemDropZone;
 };
 
-class DATAHIERARCHYEDITOR_API SDataHierarchyEditor : public SCompoundWidget, public FNotifyHook
+class SDataHierarchyEditor : public SCompoundWidget, public FNotifyHook
 {
 	DECLARE_DELEGATE_RetVal_OneParam(TSharedRef<SWidget>, FOnGenerateRowContentWidget, TSharedRef<FHierarchyElementViewModel> HierarchyElement);
 	DECLARE_DELEGATE_RetVal_OneParam(TSharedRef<SWidget>, FOnGenerateCustomDetailsPanelNameWidget, TSharedPtr<FHierarchyElementViewModel> HierarchyElement);
@@ -115,82 +117,82 @@ class DATAHIERARCHYEDITOR_API SDataHierarchyEditor : public SCompoundWidget, pub
 	SLATE_END_ARGS()
 
 public:
-	void Construct(const FArguments& InArgs, TObjectPtr<UDataHierarchyViewModelBase> InHierarchyViewModel);
-	virtual ~SDataHierarchyEditor() override;
-
-	void RefreshSourceItems();
-	void RefreshAllViews(bool bFullRefresh = false);
-	void RequestRefreshAllViewsNextFrame(bool bFullRefresh = false);
-	void RefreshSourceView(bool bFullRefresh = false) const;
-	void RequestRefreshSourceViewNextFrame(bool bFullRefresh = false);
-	void RefreshHierarchyView(bool bFullRefresh = false) const;
-	void RequestRefreshHierarchyViewNextFrame(bool bFullRefresh = false);
-	void RefreshSectionsView();
-	void RequestRefreshSectionsViewNextFrame();
-
-	void NavigateToHierarchyElement(FHierarchyElementIdentity Identity) const;
-	void NavigateToHierarchyElement(TSharedPtr<FHierarchyElementViewModel> Item) const;
-	bool IsItemSelected(TSharedPtr<FHierarchyElementViewModel> Item) const;
+	UE_API void Construct(const FArguments& InArgs, TObjectPtr<UDataHierarchyViewModelBase> InHierarchyViewModel);
+	UE_API virtual ~SDataHierarchyEditor() override;
+
+	UE_API void RefreshSourceItems();
+	UE_API void RefreshAllViews(bool bFullRefresh = false);
+	UE_API void RequestRefreshAllViewsNextFrame(bool bFullRefresh = false);
+	UE_API void RefreshSourceView(bool bFullRefresh = false) const;
+	UE_API void RequestRefreshSourceViewNextFrame(bool bFullRefresh = false);
+	UE_API void RefreshHierarchyView(bool bFullRefresh = false) const;
+	UE_API void RequestRefreshHierarchyViewNextFrame(bool bFullRefresh = false);
+	UE_API void RefreshSectionsView();
+	UE_API void RequestRefreshSectionsViewNextFrame();
+
+	UE_API void NavigateToHierarchyElement(FHierarchyElementIdentity Identity) const;
+	UE_API void NavigateToHierarchyElement(TSharedPtr<FHierarchyElementViewModel> Item) const;
+	UE_API bool IsItemSelected(TSharedPtr<FHierarchyElementViewModel> Item) const;
 	
 private:
 	// need to do this to enable focus so we can handle shortcuts
 	virtual bool SupportsKeyboardFocus() const override { return true; }
 	
-	TSharedRef<ITableRow> GenerateSourceItemRow(TSharedPtr<FHierarchyElementViewModel> HierarchyItem, const TSharedRef<STableViewBase>& TableViewBase);
-	TSharedRef<ITableRow> GenerateHierarchyItemRow(TSharedPtr<FHierarchyElementViewModel> HierarchyItem, const TSharedRef<STableViewBase>& TableViewBase);
+	UE_API TSharedRef<ITableRow> GenerateSourceItemRow(TSharedPtr<FHierarchyElementViewModel> HierarchyItem, const TSharedRef<STableViewBase>& TableViewBase);
+	UE_API TSharedRef<ITableRow> GenerateHierarchyItemRow(TSharedPtr<FHierarchyElementViewModel> HierarchyItem, const TSharedRef<STableViewBase>& TableViewBase);
 
-	bool FilterForSourceSection(TSharedPtr<const FHierarchyElementViewModel> ItemViewModel) const;
+	UE_API bool FilterForSourceSection(TSharedPtr<const FHierarchyElementViewModel> ItemViewModel) const;
 private:
-	void Reinitialize();
+	UE_API void Reinitialize();
 
-	void BindToHierarchyRootViewModel();
-	void UnbindFromHierarchyRootViewModel() const;
+	UE_API void BindToHierarchyRootViewModel();
+	UE_API void UnbindFromHierarchyRootViewModel() const;
 	
 	/** Source items reflect the base, unedited status of items to edit into a hierarchy */
-	const TArray<TSharedPtr<FHierarchyElementViewModel>>& GetSourceItems() const;
+	UE_API const TArray<TSharedPtr<FHierarchyElementViewModel>>& GetSourceItems() const;
 	
-	bool IsDetailsPanelEditingAllowed() const;
+	UE_API bool IsDetailsPanelEditingAllowed() const;
 	
-	void ClearSourceItems() const;
+	UE_API void ClearSourceItems() const;
 	
-	void RequestRenameSelectedItem();
-	bool CanRequestRenameSelectedItem() const;
+	UE_API void RequestRenameSelectedItem();
+	UE_API bool CanRequestRenameSelectedItem() const;
 
-	void DeleteItems(TArray<TSharedPtr<FHierarchyElementViewModel>> ItemsToDelete) const;
-	void DeleteSelectedHierarchyItems() const;
-	bool CanDeleteSelectedElements() const;
+	UE_API void DeleteItems(TArray<TSharedPtr<FHierarchyElementViewModel>> ItemsToDelete) const;
+	UE_API void DeleteSelectedHierarchyItems() const;
+	UE_API bool CanDeleteSelectedElements() const;
 
-	void NavigateToMatchingHierarchyElementFromSelectedSourceElement() const;
-	bool CanNavigateToMatchingHierarchyElementFromSelectedSourceElement() const;
+	UE_API void NavigateToMatchingHierarchyElementFromSelectedSourceElement() const;
+	UE_API bool CanNavigateToMatchingHierarchyElementFromSelectedSourceElement() const;
 
-	void DeleteActiveSection() const;
-	bool CanDeleteActiveSection() const;
+	UE_API void DeleteActiveSection() const;
+	UE_API bool CanDeleteActiveSection() const;
 
-	void OnElementAdded(TSharedPtr<FHierarchyElementViewModel> AddedItem);
-	void OnHierarchySectionActivated(TSharedPtr<FHierarchySectionViewModel> Section);
-	void OnSourceSectionActivated(TSharedPtr<FHierarchySectionViewModel> Section);
-	void OnHierarchySectionAdded(TSharedPtr<FHierarchySectionViewModel> AddedSection);
-	void OnHierarchySectionDeleted(TSharedPtr<FHierarchySectionViewModel> DeletedSection);
+	UE_API void OnElementAdded(TSharedPtr<FHierarchyElementViewModel> AddedItem);
+	UE_API void OnHierarchySectionActivated(TSharedPtr<FHierarchySectionViewModel> Section);
+	UE_API void OnSourceSectionActivated(TSharedPtr<FHierarchySectionViewModel> Section);
+	UE_API void OnHierarchySectionAdded(TSharedPtr<FHierarchySectionViewModel> AddedSection);
+	UE_API void OnHierarchySectionDeleted(TSharedPtr<FHierarchySectionViewModel> DeletedSection);
 
-	void SetActiveSourceSection(TSharedPtr<struct FHierarchySectionViewModel>);
-	TSharedPtr<FHierarchySectionViewModel> GetActiveSourceSection() const;
-	UHierarchySection* GetActiveSourceSectionData() const;
+	UE_API void SetActiveSourceSection(TSharedPtr<struct FHierarchySectionViewModel>);
+	UE_API TSharedPtr<FHierarchySectionViewModel> GetActiveSourceSection() const;
+	UE_API UHierarchySection* GetActiveSourceSectionData() const;
 	
-	void OnSelectionChanged(TSharedPtr<FHierarchyElementViewModel> HierarchyItem, ESelectInfo::Type Type, bool bFromHierarchy) const;
+	UE_API void OnSelectionChanged(TSharedPtr<FHierarchyElementViewModel> HierarchyItem, ESelectInfo::Type Type, bool bFromHierarchy) const;
 
-	void RunSourceSearch();
-	void OnSourceSearchTextChanged(const FText& Text);
-	void OnSourceSearchTextCommitted(const FText& Text, ETextCommit::Type CommitType);
-	void OnSearchButtonClicked(SSearchBox::SearchDirection SearchDirection);
+	UE_API void RunSourceSearch();
+	UE_API void OnSourceSearchTextChanged(const FText& Text);
+	UE_API void OnSourceSearchTextCommitted(const FText& Text, ETextCommit::Type CommitType);
+	UE_API void OnSearchButtonClicked(SSearchBox::SearchDirection SearchDirection);
 
-	virtual FReply OnKeyDown(const FGeometry& MyGeometry, const FKeyEvent& InKeyEvent) override;
-	virtual FReply OnMouseButtonDown(const FGeometry& MyGeometry, const FPointerEvent& MouseEvent) override;
-	virtual FReply OnMouseButtonUp(const FGeometry& MyGeometry, const FPointerEvent& MouseEvent) override;
+	UE_API virtual FReply OnKeyDown(const FGeometry& MyGeometry, const FKeyEvent& InKeyEvent) override;
+	UE_API virtual FReply OnMouseButtonDown(const FGeometry& MyGeometry, const FPointerEvent& MouseEvent) override;
+	UE_API virtual FReply OnMouseButtonUp(const FGeometry& MyGeometry, const FPointerEvent& MouseEvent) override;
 
-	FReply OnAddCategoryClicked() const;
-	FReply OnAddSectionClicked() const;
+	UE_API FReply OnAddCategoryClicked() const;
+	UE_API FReply OnAddSectionClicked() const;
 
-	TSharedPtr<SWidget> SummonContextMenuForSelectedRows(bool bFromHierarchy) const;
+	UE_API TSharedPtr<SWidget> SummonContextMenuForSelectedRows(bool bFromHierarchy) const;
 
 	struct FSearchItem
 	{
@@ -210,22 +212,22 @@ private:
 	};
 
 	/** This will recursively generated parent chain paths for all items within the given root. Used for expansion purposes. */
-	void GenerateSearchItems(TSharedRef<FHierarchyElementViewModel> Root, TArray<TSharedPtr<FHierarchyElementViewModel>> ParentChain, TArray<FSearchItem>& OutSearchItems);
-	void ExpandSourceSearchResults();
-	void SelectNextSourceSearchResult();
-	void SelectPreviousSourceSearchResult();
-	TOptional<SSearchBox::FSearchResultData> GetSearchResultData() const;
+	UE_API void GenerateSearchItems(TSharedRef<FHierarchyElementViewModel> Root, TArray<TSharedPtr<FHierarchyElementViewModel>> ParentChain, TArray<FSearchItem>& OutSearchItems);
+	UE_API void ExpandSourceSearchResults();
+	UE_API void SelectNextSourceSearchResult();
+	UE_API void SelectPreviousSourceSearchResult();
+	UE_API TOptional<SSearchBox::FSearchResultData> GetSearchResultData() const;
 	
-	FHierarchyElementViewModel::FCanPerformActionResults CanDropOnRoot(TSharedPtr<FHierarchyElementViewModel> DraggedItem) const;
+	UE_API FHierarchyElementViewModel::FCanPerformActionResults CanDropOnRoot(TSharedPtr<FHierarchyElementViewModel> DraggedItem) const;
 
 	/** Callback functions for the root widget */
-	FReply HandleHierarchyRootDrop(const FGeometry&, const FDragDropEvent& DragDropEvent) const;
-	bool OnCanDropOnRoot(TSharedPtr<FDragDropOperation> DragDropOperation) const;
-	void OnRootDragEnter(const FDragDropEvent& DragDropEvent) const;
-	void OnRootDragLeave(const FDragDropEvent& DragDropEvent) const;
-	FSlateColor GetRootIconColor() const;
+	UE_API FReply HandleHierarchyRootDrop(const FGeometry&, const FDragDropEvent& DragDropEvent) const;
+	UE_API bool OnCanDropOnRoot(TSharedPtr<FDragDropOperation> DragDropOperation) const;
+	UE_API void OnRootDragEnter(const FDragDropEvent& DragDropEvent) const;
+	UE_API void OnRootDragLeave(const FDragDropEvent& DragDropEvent) const;
+	UE_API FSlateColor GetRootIconColor() const;
 
-	virtual void NotifyPostChange(const FPropertyChangedEvent& PropertyChangedEvent, FProperty* PropertyThatChanged) override;
+	UE_API virtual void NotifyPostChange(const FPropertyChangedEvent& PropertyChangedEvent, FProperty* PropertyThatChanged) override;
 private:
 	TWeakObjectPtr<UDataHierarchyViewModelBase> HierarchyViewModel;
 	
@@ -256,7 +258,7 @@ private:
 	TSharedPtr<FActiveTimerHandle> RefreshSectionsViewNextFrameHandle;
 };
 
-class DATAHIERARCHYEDITOR_API SHierarchyCategory : public SCompoundWidget
+class SHierarchyCategory : public SCompoundWidget
 {
 	SLATE_BEGIN_ARGS(SHierarchyCategory)
 	{}
@@ -264,15 +266,17 @@ class DATAHIERARCHYEDITOR_API SHierarchyCategory : public SCompoundWidget
 	SLATE_END_ARGS()
 
 public:
-	void Construct(const FArguments& InArgs, TSharedPtr<struct FHierarchyCategoryViewModel> InCategoryViewModel);
+	UE_API void Construct(const FArguments& InArgs, TSharedPtr<struct FHierarchyCategoryViewModel> InCategoryViewModel);
 
-	void EnterEditingMode() const;
-	bool OnVerifyCategoryRename(const FText& NewName, FText& OutTooltip) const;
+	UE_API void EnterEditingMode() const;
+	UE_API bool OnVerifyCategoryRename(const FText& NewName, FText& OutTooltip) const;
 	
-	FText GetCategoryText() const;
-	void OnRenameCategory(const FText& NewText, ETextCommit::Type) const;
+	UE_API FText GetCategoryText() const;
+	UE_API void OnRenameCategory(const FText& NewText, ETextCommit::Type) const;
 	
 private:
 	TWeakPtr<struct FHierarchyCategoryViewModel> CategoryViewModelWeak;
 	TSharedPtr<SInlineEditableTextBlock> InlineEditableTextBlock;
 };
+
+#undef UE_API
