- Check for the timer manager being valid
- Content can be loaded before it is in some situations

#tests Launch the editor in some test projects
#rnx

diff --git a/Engine/Source/Editor/DataHierarchyEditor/Private/DataHierarchyViewModelBase.cpp b/Engine/Source/Editor/DataHierarchyEditor/Private/DataHierarchyViewModelBase.cpp
--- a/Engine/Source/Editor/DataHierarchyEditor/Private/DataHierarchyViewModelBase.cpp
+++ b/Engine/Source/Editor/DataHierarchyEditor/Private/DataHierarchyViewModelBase.cpp
@@ -572,7 +572,7 @@ void UDataHierarchyViewModelBase::ForceFullRefreshOnTimer()
 
 void UDataHierarchyViewModelBase::RequestFullRefreshNextFrame()
 {
-	if(!FullRefreshNextFrameHandle.IsValid() && GEditor != nullptr)
+	if(!FullRefreshNextFrameHandle.IsValid() && GEditor != nullptr && GEditor->IsTimerManagerValid())
 	{
 		FullRefreshNextFrameHandle = GEditor->GetTimerManager()->SetTimerForNextTick(FTimerDelegate::CreateUObject(this, &UDataHierarchyViewModelBase::ForceFullRefreshOnTimer));
 	}
