- Handle nullptr when accessing ActiveHierarchySection in a few places

#rnx
#tests Tested with crash case Dan had

diff --git a/Engine/Source/Editor/DataHierarchyEditor/Private/DataHierarchyViewModelBase.cpp b/Engine/Source/Editor/DataHierarchyEditor/Private/DataHierarchyViewModelBase.cpp
--- a/Engine/Source/Editor/DataHierarchyEditor/Private/DataHierarchyViewModelBase.cpp
+++ b/Engine/Source/Editor/DataHierarchyEditor/Private/DataHierarchyViewModelBase.cpp
@@ -767,12 +767,14 @@ TSharedPtr<FHierarchySectionViewModel> UDataHierarchyViewModelBase::GetActiveHie
 
 UHierarchySection* UDataHierarchyViewModelBase::GetActiveHierarchySectionData() const
 {
-	return ActiveHierarchySection.Pin()->GetDataMutable<UHierarchySection>();
+	TSharedPtr<FHierarchySectionViewModel> ActiveSection = GetActiveHierarchySection();
+	return ActiveSection.IsValid() ? ActiveSection->GetDataMutable<UHierarchySection>() : nullptr;
 }
 
 bool UDataHierarchyViewModelBase::IsHierarchySectionActive(const UHierarchySection* Section) const
 {
-	return ActiveHierarchySection.Pin()->GetData() == Section;
+	TSharedPtr<FHierarchySectionViewModel> ActiveSection = GetActiveHierarchySection();
+	return ActiveSection.IsValid() ? ActiveSection->GetData() == Section : false;
 }
 
 FString UDataHierarchyViewModelBase::OnElementToStringDebug(TSharedPtr<FHierarchyElementViewModel> ElementViewModel) const
