Data Hierarchy Editor Upgrade
- Added new 'CanContain' API that allows elements to specify whether they can contain specific types. This allows abstracting away repeats of drag & drop logic.
-- OnDroppedOnInternal no longer exists. OnDroppedOn now has a default implementation (previously was just calling OnDroppedOnInternal) handling some common checks for you, including the new 'CanContain' call. This can optionally be overridden by using OnDroppedOnOverride.
- Added virtual 'PostAddFixUp' function that can be overridden in case a hierarchy object has been moved or added to some place. For example used to set up section data correctly after moving a category around.
- Renamed 'GetDataForEditing' to 'GetObjectForEditing'
- Added "SelectRootInsteadOfNone" option. Useful if the root is supposed to have some data on it or the root view model points at an external object for editing.
- View models can not specify their own row styles. Previously this was set once depending on item or category.
- Sources panel has become optional
- Using the new 'CanContain' API and the new 'GetAdditionalTypesToAddInUi` function allows to add an 'Add' button with which new elements in the hierarchy can be constructed. This is useful if you aren't sorting pre-existing elements into a hierarchy but instead construct new elements from scratch.
- The 'All' section & view model logic has been refactored. Previously the All section used a section view model with no backing data, indicating "nullptr" section data means All section. This led to some inconsistencies. Instead, now All section view models get a transient section object to work with.
- Renamed FCanPerformActionResults to FResultWithUserFeedback because the feedback might not be used on an actual action; this way it's more generic and can be used to surface information to the user
- Deleted SHierarchyCategory widget and replaced with SHierarchyElement widget that works for both categories and items. Default widget that supports the rename API.

diff --git a/Engine/Plugins/FX/Niagara/Source/NiagaraEditor/Private/NiagaraEditorUtilities.cpp b/Engine/Plugins/FX/Niagara/Source/NiagaraEditor/Private/NiagaraEditorUtilities.cpp
--- a/Engine/Plugins/FX/Niagara/Source/NiagaraEditor/Private/NiagaraEditorUtilities.cpp
+++ b/Engine/Plugins/FX/Niagara/Source/NiagaraEditor/Private/NiagaraEditorUtilities.cpp
@@ -4852,10 +4852,16 @@ TSharedRef<SWidget> FNiagaraEditorUtilities::HierarchyEditor::Scripts::GenerateR
 	if(HierarchyItem->GetDataMutable()->IsA<UHierarchyCategory>())
 	{
 		TSharedRef<FHierarchyCategoryViewModel> TreeViewCategory = StaticCastSharedRef<FHierarchyCategoryViewModel>(HierarchyItem);
-		return SNew(SHierarchyCategory, TreeViewCategory);
+		return SNew(SHierarchyElement, TreeViewCategory)
+			.Style(&FDataHierarchyEditorStyle::Get().GetWidgetStyle<FInlineEditableTextBlockStyle>("HierarchyEditor.Category"));
 	}
 	else if(const UNiagaraHierarchyScriptParameter* ScriptParameter = Cast<UNiagaraHierarchyScriptParameter>(HierarchyItem->GetData()))
 	{
+		if(ScriptParameter->GetScriptVariable() == nullptr)
+		{
+			return SNew(STextBlock).Text(LOCTEXT("InvalidHierarchyScriptParameter", "Invalid Script Hierarchy Parameter"));
+		}
+		
 		FNiagaraParameterUtilities::FNiagaraParameterWidgetOptions WidgetOptions;
 		WidgetOptions.bAddTypeIcon = true;
 		WidgetOptions.bShowEditConditionIcon = true;
diff --git a/Engine/Plugins/FX/Niagara/Source/NiagaraEditor/Private/Toolkits/SystemToolkitModes/NiagaraSystemToolkitModeBase.cpp b/Engine/Plugins/FX/Niagara/Source/NiagaraEditor/Private/Toolkits/SystemToolkitModes/NiagaraSystemToolkitModeBase.cpp
--- a/Engine/Plugins/FX/Niagara/Source/NiagaraEditor/Private/Toolkits/SystemToolkitModes/NiagaraSystemToolkitModeBase.cpp
+++ b/Engine/Plugins/FX/Niagara/Source/NiagaraEditor/Private/Toolkits/SystemToolkitModes/NiagaraSystemToolkitModeBase.cpp
@@ -879,7 +879,8 @@ TSharedRef<SWidget> GenerateRowContentForUserParameterHierarchyEditor(TSharedRef
 	if(HierarchyItem->GetDataMutable()->IsA<UHierarchyCategory>())
 	{
 		TSharedRef<FHierarchyCategoryViewModel> TreeViewCategory = StaticCastSharedRef<FHierarchyCategoryViewModel>(HierarchyItem);
-		return SNew(SHierarchyCategory, TreeViewCategory);
+		return SNew(SHierarchyElement, TreeViewCategory)
+			.Style(&FDataHierarchyEditorStyle::Get().GetWidgetStyle<FInlineEditableTextBlockStyle>("HierarchyEditor.Category"));
 	}
 	else if(const UNiagaraHierarchyUserParameter* UserParameter = Cast<UNiagaraHierarchyUserParameter>(HierarchyItem->GetData()))
 	{
@@ -1031,7 +1032,8 @@ TSharedRef<SWidget> GenerateRowContentForSummaryViewHierarchyEditor(TSharedRef<F
 	else if(HierarchyItem->GetDataMutable()->IsA<UHierarchyCategory>())
 	{
 		TSharedRef<FHierarchyCategoryViewModel> TreeViewCategory = StaticCastSharedRef<FHierarchyCategoryViewModel>(HierarchyItem);
-		return SNew(SHierarchyCategory, TreeViewCategory);
+		return SNew(SHierarchyElement, TreeViewCategory)
+			.Style(&FDataHierarchyEditorStyle::Get().GetWidgetStyle<FInlineEditableTextBlockStyle>("HierarchyEditor.Category"));
 	}
 	else if(const UHierarchyItem* Item = Cast<UHierarchyItem>(HierarchyItem->GetData()))
 	{
diff --git a/Engine/Plugins/FX/Niagara/Source/NiagaraEditor/Private/ViewModels/HierarchyEditor/NiagaraScriptParametersHierarchyViewModel.cpp b/Engine/Plugins/FX/Niagara/Source/NiagaraEditor/Private/ViewModels/HierarchyEditor/NiagaraScriptParametersHierarchyViewModel.cpp
--- a/Engine/Plugins/FX/Niagara/Source/NiagaraEditor/Private/ViewModels/HierarchyEditor/NiagaraScriptParametersHierarchyViewModel.cpp
+++ b/Engine/Plugins/FX/Niagara/Source/NiagaraEditor/Private/ViewModels/HierarchyEditor/NiagaraScriptParametersHierarchyViewModel.cpp
@@ -255,20 +255,20 @@ TSharedRef<SWidget> FNiagaraHierarchyScriptParameterDragDropOp::CreateCustomDeco
 	return SNullWidget::NullWidget;
 }
 
-FHierarchyElementViewModel::FCanPerformActionResults FNiagaraHierarchyScriptCategoryViewModel::CanDropOnInternal(TSharedPtr<FHierarchyElementViewModel> DraggedElement, EItemDropZone ItemDropZone)
+FHierarchyElementViewModel::FResultWithUserFeedback FNiagaraHierarchyScriptCategoryViewModel::CanDropOnInternal(TSharedPtr<FHierarchyElementViewModel> DraggedElement, EItemDropZone ItemDropZone)
 {
-	if(IsEditableByUser().bCanPerform == false)
+	if(IsEditableByUser().bResult == false)
 	{
 		return false;
 	}
 	
-	FCanPerformActionResults CanPerformActionResults = FHierarchyCategoryViewModel::CanDropOnInternal(DraggedElement, ItemDropZone);
+	FResultWithUserFeedback CanPerformActionResults = FHierarchyCategoryViewModel::CanDropOnInternal(DraggedElement, ItemDropZone);
 
 	// We allow dropping parameters above/below categories
 	if(ItemDropZone != EItemDropZone::OntoItem && DraggedElement->GetData()->IsA<UNiagaraHierarchyScriptParameter>())
 	{
-		CanPerformActionResults.CanPerformMessage = FText::GetEmpty();
-		CanPerformActionResults.bCanPerform = true;
+		CanPerformActionResults.UserFeedback = FText::GetEmpty();
+		CanPerformActionResults.bResult = true;
 	}
 
 	return CanPerformActionResults;
@@ -293,16 +293,35 @@ bool FNiagaraHierarchyScriptParameterViewModel::DoesExternalDataStillExist(const
 	return true;
 }
 
-FHierarchyElementViewModel::FCanPerformActionResults FNiagaraHierarchyScriptParameterViewModel::CanDropOnInternal(TSharedPtr<FHierarchyElementViewModel> DraggedItem, EItemDropZone ItemDropZone)
+FHierarchyElementViewModel::FResultWithUserFeedback FNiagaraHierarchyScriptParameterViewModel::CanContainInternal(TSubclassOf<UHierarchyElement> InHierarchyElementType)
+{
+	// We support nested inputs up to 1 layer deep
+	if(InHierarchyElementType->IsChildOf<UNiagaraHierarchyScriptParameter>())
+	{
+		if(Parent.IsValid() && Parent.Pin()->GetData()->IsA<UNiagaraHierarchyScriptParameter>())
+		{
+			FResultWithUserFeedback CanContainResult(false);
+			FText BaseMessage = LOCTEXT("InputOnInputNestedChildTooDeep", "Can't nest new input {0} under input {1}.\nChildren inputs can only have one layer of depth! Add it to the parent input!");
+			CanContainResult.UserFeedback = FText::FormatOrdered(BaseMessage, InHierarchyElementType->GetDisplayNameText(), ToStringAsText());
+			return CanContainResult;
+		}
+				
+		return true;
+	}
+	
+	return false;
+}
+
+FHierarchyElementViewModel::FResultWithUserFeedback FNiagaraHierarchyScriptParameterViewModel::CanDropOnInternal(TSharedPtr<FHierarchyElementViewModel> DraggedItem, EItemDropZone ItemDropZone)
 {
 	// if the input isn't editable, we don't allow any drops on/above/below the item.
 	// Even though it technically works, the merge process will only re-add the item at the end and not preserve order so there is no point in allowing dropping above/below
-	if(IsEditableByUser().bCanPerform == false)
+	if(IsEditableByUser().bResult == false)
 	{
 		return false;
 	}
 
-	FCanPerformActionResults AllowDrop(false);
+	FResultWithUserFeedback AllowDrop(false);
 	
 	TSharedPtr<FHierarchyElementViewModel> TargetDropItem = AsShared();
 
@@ -320,8 +339,8 @@ FHierarchyElementViewModel::FCanPerformActionResults FNiagaraHierarchyScriptPara
 				if(DraggedItem->GetData()->DoesOneChildExist<UNiagaraHierarchyScriptParameter>())
 				{
 					FText BaseMessage = LOCTEXT("DraggedItemHasChild_NotAllowed", "Can't nest input {0} under input {1}.\nRemove children of {2} first. ");
-					AllowDrop.CanPerformMessage = FText::FormatOrdered(BaseMessage, DraggedItem->ToStringAsText(), TargetDropItem->ToStringAsText(), DraggedItem->ToStringAsText());
-					AllowDrop.bCanPerform = false;
+					AllowDrop.UserFeedback = FText::FormatOrdered(BaseMessage, DraggedItem->ToStringAsText(), TargetDropItem->ToStringAsText(), DraggedItem->ToStringAsText());
+					AllowDrop.bResult = false;
 					return AllowDrop;
 				}
 				
@@ -329,14 +348,14 @@ FHierarchyElementViewModel::FCanPerformActionResults FNiagaraHierarchyScriptPara
 				if(TargetDropItem->GetParent().Pin()->GetData<UNiagaraHierarchyScriptParameter>())
 				{
 					FText BaseMessage = LOCTEXT("DroppingInputOnInputNestedChildTooDeep", "Can't nest input {0} under input {1}.\nChildren inputs can only have one layer of depth!");
-					AllowDrop.CanPerformMessage = FText::FormatOrdered(BaseMessage, DraggedItem->ToStringAsText(), TargetDropItem->ToStringAsText());
-					AllowDrop.bCanPerform = false;
+					AllowDrop.UserFeedback = FText::FormatOrdered(BaseMessage, DraggedItem->ToStringAsText(), TargetDropItem->ToStringAsText());
+					AllowDrop.bResult = false;
 					return AllowDrop;
 				}
 				
 				FText BaseMessage = LOCTEXT("DroppingInputOnInputNestedChild", "This will nest input {0} under input {1}");
-				AllowDrop.CanPerformMessage = FText::FormatOrdered(BaseMessage, DraggedItem->ToStringAsText(), TargetDropItem->ToStringAsText());
-				AllowDrop.bCanPerform = true;
+				AllowDrop.UserFeedback = FText::FormatOrdered(BaseMessage, DraggedItem->ToStringAsText(), TargetDropItem->ToStringAsText());
+				AllowDrop.bResult = true;
 				return AllowDrop;
 			}
 		}
@@ -345,15 +364,15 @@ FHierarchyElementViewModel::FCanPerformActionResults FNiagaraHierarchyScriptPara
 			// if the dragged item is an input, we generally allow above/below, even for nested child inputs
 			if(DraggedItem->GetData()->IsA<UNiagaraHierarchyScriptParameter>())
 			{
-				AllowDrop.bCanPerform = true;
+				AllowDrop.bResult = true;
 			}
 			// if the dragged item is a category, we generally allow putting it above/below other parameters, but not above/below child parameters
 			else if(DraggedItem->GetData()->IsA<UHierarchyCategory>())
 			{
-				AllowDrop.bCanPerform = true;
+				AllowDrop.bResult = true;
 				if(TargetDropItem->GetData()->IsA(UNiagaraHierarchyScriptParameter::StaticClass()) && TargetDropItem->GetParent().Pin()->GetData()->IsA(UNiagaraHierarchyScriptParameter::StaticClass()))
 				{
-					AllowDrop.bCanPerform = false;
+					AllowDrop.bResult = false;
 				}
 			}
 		}
@@ -362,31 +381,4 @@ FHierarchyElementViewModel::FCanPerformActionResults FNiagaraHierarchyScriptPara
 	return AllowDrop;
 }
 
-void FNiagaraHierarchyScriptParameterViewModel::OnDroppedOnInternal(TSharedPtr<FHierarchyElementViewModel> DroppedItem, EItemDropZone ItemDropZone)
-{
-	if(ItemDropZone != EItemDropZone::OntoItem)
-	{
-		return FHierarchyItemViewModel::OnDroppedOnInternal(DroppedItem, ItemDropZone);
-	}
-	else
-	{
-		FScopedTransaction Transaction(LOCTEXT("Transaction_AddedChildInput", "Added child input"));
-		HierarchyViewModel->GetHierarchyRoot()->Modify();
-		
-		if(DroppedItem->IsForHierarchy() == false)
-		{
-			TSharedRef<FHierarchyElementViewModel> AddedItemViewModel = DuplicateToThis(DroppedItem);
-			AddedItemViewModel->SyncViewModelsToData();
-		}
-		else
-		{
-			TSharedRef<FHierarchyElementViewModel> ReparentedViewModel = ReparentToThis(DroppedItem);
-			ReparentedViewModel->SyncViewModelsToData();
-		}
-
-		HierarchyViewModel->RefreshHierarchyView();
-		HierarchyViewModel->RefreshSourceView();
-	}
-}
-
 #undef LOCTEXT_NAMESPACE
diff --git a/Engine/Plugins/FX/Niagara/Source/NiagaraEditor/Private/ViewModels/HierarchyEditor/NiagaraSummaryViewViewModel.cpp b/Engine/Plugins/FX/Niagara/Source/NiagaraEditor/Private/ViewModels/HierarchyEditor/NiagaraSummaryViewViewModel.cpp
--- a/Engine/Plugins/FX/Niagara/Source/NiagaraEditor/Private/ViewModels/HierarchyEditor/NiagaraSummaryViewViewModel.cpp
+++ b/Engine/Plugins/FX/Niagara/Source/NiagaraEditor/Private/ViewModels/HierarchyEditor/NiagaraSummaryViewViewModel.cpp
@@ -184,17 +184,17 @@ FString FNiagaraFunctionViewModel::ToString() const
 	return TEXT("Unknown");
 }
 
-FHierarchyElementViewModel::FCanPerformActionResults FNiagaraFunctionViewModel::IsEditableByUser()
+FHierarchyElementViewModel::FResultWithUserFeedback FNiagaraFunctionViewModel::IsEditableByUser()
 {
 	if(bIsDynamicInput)
 	{
-		FCanPerformActionResults CanEditResults(false);
-		CanEditResults.CanPerformMessage = LOCTEXT("DynamicInputCantBeDragged", "You can not drag entire Dynamic Inputs. Either drag the entire module input, or individual inputs of the Dynamic Input");
+		FResultWithUserFeedback CanEditResults(false);
+		CanEditResults.UserFeedback = LOCTEXT("DynamicInputCantBeDragged", "You can not drag entire Dynamic Inputs. Either drag the entire module input, or individual inputs of the Dynamic Input");
 		return CanEditResults;
 	}
 	
-	FCanPerformActionResults CanEditResults(IsFromBaseEmitter() == false);
-	CanEditResults.CanPerformMessage = CanEditResults.bCanPerform == false ? LOCTEXT("ModuleIsFromBaseEmitter", "This module was added in the parent emitter and can not be edited.") : FText::GetEmpty();
+	FResultWithUserFeedback CanEditResults(IsFromBaseEmitter() == false);
+	CanEditResults.UserFeedback = CanEditResults.bResult == false ? LOCTEXT("ModuleIsFromBaseEmitter", "This module was added in the parent emitter and can not be edited.") : FText::GetEmpty();
 	return CanEditResults;
 }
 
@@ -341,9 +341,9 @@ void FNiagaraFunctionViewModel::RefreshChildrenInputs(bool bClearCache) const
 	}
 }
 
-FHierarchyElementViewModel::FCanPerformActionResults FNiagaraFunctionViewModel::CanDropOnInternal(TSharedPtr<FHierarchyElementViewModel> NiagaraHierarchyItemViewModelBase, EItemDropZone ItemDropZone)
+FHierarchyElementViewModel::FResultWithUserFeedback FNiagaraFunctionViewModel::CanDropOnInternal(TSharedPtr<FHierarchyElementViewModel> NiagaraHierarchyItemViewModelBase, EItemDropZone ItemDropZone)
 {
-	if(IsEditableByUser().bCanPerform == false)
+	if(IsEditableByUser().bResult == false)
 	{
 		return false;
 	}
@@ -351,8 +351,8 @@ FHierarchyElementViewModel::FCanPerformActionResults FNiagaraFunctionViewModel::
 	// we don't allow any items to be added directly onto the module as it's self managing
 	if(ItemDropZone == EItemDropZone::OntoItem)
 	{
-		FCanPerformActionResults Results(false);
-		Results.CanPerformMessage = LOCTEXT("CanDropOnModuleDragMessage", "You can not add any items to a module directly. Please create a category, which can contain arbitrary items.");
+		FResultWithUserFeedback Results(false);
+		Results.UserFeedback = LOCTEXT("CanDropOnModuleDragMessage", "You can not add any items to a module directly. Please create a category, which can contain arbitrary items.");
 		return Results;
 	}
 	
@@ -369,16 +369,16 @@ const UHierarchySection* FNiagaraFunctionViewModel::GetSectionInternal() const
 	return Section.IsValid() ? Section.Get() : nullptr;
 }
 
-FHierarchyElementViewModel::FCanPerformActionResults FNiagaraModuleInputViewModel::CanDropOnInternal(TSharedPtr<FHierarchyElementViewModel> DraggedItem, EItemDropZone ItemDropZone)
+FHierarchyElementViewModel::FResultWithUserFeedback FNiagaraModuleInputViewModel::CanDropOnInternal(TSharedPtr<FHierarchyElementViewModel> DraggedItem, EItemDropZone ItemDropZone)
 {
 	// if the input isn't editable, we don't allow any drops on/above/below the item.
 	// Even though it technically works, the merge process will only re-add the item at the end and not preserve order so there is no point in allowing dropping above/below
-	if(IsEditableByUser().bCanPerform == false)
+	if(IsEditableByUser().bResult == false)
 	{
 		return false;
 	}
 
-	FCanPerformActionResults AllowDrop(false);
+	FResultWithUserFeedback AllowDrop(false);
 	
 	TSharedPtr<FHierarchyElementViewModel> TargetDropItem = AsShared();
 
@@ -395,14 +395,14 @@ FHierarchyElementViewModel::FCanPerformActionResults FNiagaraModuleInputViewMode
 				if(DraggedItem->GetData()->GetChildren().Num() > 0)
 				{
 					FText BaseMessage = LOCTEXT("DroppingInputOnInputWillEmptyChildren", "Input {0} has child inputs. Dropping the input here will remove these children as we only allow nested inputs one level deep.");
-					AllowDrop.CanPerformMessage = FText::FormatOrdered(BaseMessage, DraggedItem->ToStringAsText());
-					AllowDrop.bCanPerform = true;
+					AllowDrop.UserFeedback = FText::FormatOrdered(BaseMessage, DraggedItem->ToStringAsText());
+					AllowDrop.bResult = true;
 				}
 				else
 				{
 					FText BaseMessage = LOCTEXT("DroppingInputOnInputNestedChild", "This will nest input {0} under input {1}");
-					AllowDrop.CanPerformMessage = FText::FormatOrdered(BaseMessage, DraggedItem->ToStringAsText(), TargetDropItem->ToStringAsText());
-					AllowDrop.bCanPerform = true;
+					AllowDrop.UserFeedback = FText::FormatOrdered(BaseMessage, DraggedItem->ToStringAsText(), TargetDropItem->ToStringAsText());
+					AllowDrop.bResult = true;
 				}
 			}
 		}
@@ -411,7 +411,7 @@ FHierarchyElementViewModel::FCanPerformActionResults FNiagaraModuleInputViewMode
 			// if the dragged item is an input, we generally allow above/below, even for nested child inputs
 			if(DraggedItem->GetData()->IsA<UNiagaraHierarchyModuleInput>())
 			{
-				AllowDrop.bCanPerform = true;
+				AllowDrop.bResult = true;
 			}
 			else
 			{
@@ -427,33 +427,12 @@ FHierarchyElementViewModel::FCanPerformActionResults FNiagaraModuleInputViewMode
 	return AllowDrop;
 }
 
-void FNiagaraModuleInputViewModel::OnDroppedOnInternal(TSharedPtr<FHierarchyElementViewModel> DroppedItem, EItemDropZone ItemDropZone)
+void FNiagaraModuleInputViewModel::PostAddFixup()
 {
-	if(ItemDropZone != EItemDropZone::OntoItem)
+	// If our new parent is an input, we empty out our own inputs as we only allow 1 layer of child inputs
+	if(Parent.Pin()->GetData()->IsA<UNiagaraHierarchyModuleInput>())
 	{
-		return FHierarchyItemViewModel::OnDroppedOnInternal(DroppedItem, ItemDropZone);
-	}
-	else
-	{
-		FScopedTransaction Transaction(LOCTEXT("Transaction_AddedChildInput", "Added child input"));
-		HierarchyViewModel->GetHierarchyRoot()->Modify();
-
-		// we empty out the children as technically you can drag a parent input onto another input now. We don't take these child-child inputs with us as we only allow child inputs 1 layer deep
-		if(DroppedItem->IsForHierarchy() == false)
-		{
-			TSharedRef<FHierarchyElementViewModel> AddedItemViewModel = DuplicateToThis(DroppedItem);
-			AddedItemViewModel->GetChildrenMutable().Empty();
-			AddedItemViewModel->SyncViewModelsToData();
-		}
-		else
-		{
-			TSharedRef<FHierarchyElementViewModel> ReparentedViewModel = ReparentToThis(DroppedItem);
-			ReparentedViewModel->GetChildrenMutable().Empty();
-			ReparentedViewModel->SyncViewModelsToData();
-		}
-
-		HierarchyViewModel->RefreshHierarchyView();
-		HierarchyViewModel->RefreshSourceView();
+		GetDataMutable()->GetChildrenMutable().Empty();
 	}
 }
 
@@ -486,7 +465,7 @@ TOptional<FInputData> FNiagaraModuleInputViewModel::GetInputData() const
 	return InputDataCache;
 }
 
-bool FNiagaraModuleInputViewModel::CanHaveChildren() const
+FHierarchyElementViewModel::FResultWithUserFeedback FNiagaraModuleInputViewModel::CanHaveChildren() const
 {
 	// we generally allow children inputs in the source view
 	if(IsForHierarchy() == false)
@@ -494,13 +473,26 @@ bool FNiagaraModuleInputViewModel::CanHaveChildren() const
 		return true;
 	}
 
-	// we allow module inputs to have children inputs one layer deep
-	if(Parent.IsValid() && Parent.Pin()->GetData()->IsA<UNiagaraHierarchyModuleInput>())
+	return true;
+}
+
+FHierarchyElementViewModel::FResultWithUserFeedback FNiagaraModuleInputViewModel::CanContainInternal(TSubclassOf<UHierarchyElement> InHierarchyElementType)
+{
+	// We support nested inputs up to 1 layer deep
+	if(InHierarchyElementType->IsChildOf<UNiagaraHierarchyModuleInput>())
 	{
-		return false;
+		if(Parent.IsValid() && Parent.Pin()->GetData()->IsA<UNiagaraHierarchyModuleInput>())
+		{
+			FResultWithUserFeedback CanContainResult(false);
+			FText BaseMessage = LOCTEXT("InputOnInputNestedChildTooDeep", "Can't nest new input {0} under input {1}.\nChildren inputs can only have one layer of depth! Add it to the parent input!");
+			CanContainResult.UserFeedback = FText::FormatOrdered(BaseMessage, InHierarchyElementType->GetDisplayNameText(), ToStringAsText());
+			return CanContainResult;
+		}
+				
+		return true;
 	}
-
-	return true;
+	
+	return false;
 }
 
 FString FNiagaraModuleInputViewModel::ToString() const
@@ -574,17 +566,17 @@ FText FNiagaraModuleInputViewModel::GetSummaryInputNameOverride() const
 	return GetData<UNiagaraHierarchyModuleInput>()->GetDisplayNameOverride();
 }
 
-FHierarchyElementViewModel::FCanPerformActionResults FNiagaraModuleInputViewModel::IsEditableByUser()
+FHierarchyElementViewModel::FResultWithUserFeedback FNiagaraModuleInputViewModel::IsEditableByUser()
 {
-	FCanPerformActionResults CanEditResults(IsFromBaseEmitter() == false);
-	CanEditResults.CanPerformMessage = CanEditResults.bCanPerform == false ? LOCTEXT("ModuleInputIsFromBaseEmitter", "This input was added in the parent emitter and can not be edited.") : FText::GetEmpty();
+	FResultWithUserFeedback CanEditResults(IsFromBaseEmitter() == false);
+	CanEditResults.UserFeedback = CanEditResults.bResult == false ? LOCTEXT("ModuleInputIsFromBaseEmitter", "This input was added in the parent emitter and can not be edited.") : FText::GetEmpty();
 
-	if(bIsForHierarchy && CanEditResults.bCanPerform == true)
+	if(bIsForHierarchy && CanEditResults.bResult == true)
 	{
 		if(Parent.Pin()->GetData()->IsA<UNiagaraHierarchyModule>())
 		{
-			CanEditResults.bCanPerform = false;
-			CanEditResults.CanPerformMessage = LOCTEXT("ModuleCanOnlyBeEditedDirectly", "This input can not be modified as it is inherent part of its parent module. Add this input separately if you want to modify it.");
+			CanEditResults.bResult = false;
+			CanEditResults.UserFeedback = LOCTEXT("ModuleCanOnlyBeEditedDirectly", "This input can not be modified as it is inherent part of its parent module. Add this input separately if you want to modify it.");
 		}
 	}
 	
@@ -710,11 +702,11 @@ TArray<FHierarchyElementIdentity> FNiagaraModuleInputViewModel::GetNativeChildIn
 	return ChildIdentities;
 }
 
-FHierarchyElementViewModel::FCanPerformActionResults FNiagaraAssignmentInputViewModel::CanDropOnInternal(TSharedPtr<FHierarchyElementViewModel> DraggedItem, EItemDropZone ItemDropZone)
+FHierarchyElementViewModel::FResultWithUserFeedback FNiagaraAssignmentInputViewModel::CanDropOnInternal(TSharedPtr<FHierarchyElementViewModel> DraggedItem, EItemDropZone ItemDropZone)
 {
 	// if the input isn't editable, we don't allow any drops on/above/below the item.
 	// Even though it technically works, the merge process will only re-add the item at the end and not preserve order.
-	if(IsEditableByUser().bCanPerform == false)
+	if(IsEditableByUser().bResult == false)
 	{
 		return false;
 	}
@@ -776,17 +768,17 @@ void FNiagaraAssignmentInputViewModel::ClearCache() const
 	SummaryViewModel->ClearFunctionCallNodeCache(GetData()->GetPersistentIdentity().Guids[0]);
 }
 
-FHierarchyElementViewModel::FCanPerformActionResults FNiagaraAssignmentInputViewModel::IsEditableByUser()
+FHierarchyElementViewModel::FResultWithUserFeedback FNiagaraAssignmentInputViewModel::IsEditableByUser()
 {
-	FCanPerformActionResults CanEditResults(IsFromBaseEmitter() == false);
-	CanEditResults.CanPerformMessage = CanEditResults.bCanPerform == false ? LOCTEXT("ModuleInputIsFromBaseEmitter", "This input was added in the parent emitter and can not be edited.") : FText::GetEmpty();
+	FResultWithUserFeedback CanEditResults(IsFromBaseEmitter() == false);
+	CanEditResults.UserFeedback = CanEditResults.bResult == false ? LOCTEXT("ModuleInputIsFromBaseEmitter", "This input was added in the parent emitter and can not be edited.") : FText::GetEmpty();
 
-	if(bIsForHierarchy && CanEditResults.bCanPerform == true)
+	if(bIsForHierarchy && CanEditResults.bResult == true)
 	{
 		if(Parent.Pin()->GetData()->IsA<UNiagaraHierarchyModule>())
 		{
-			CanEditResults.bCanPerform = false;
-			CanEditResults.CanPerformMessage = LOCTEXT("ModuleCanOnlyBeEditedDirectly", "This input can not be modified as it is inherent part of its parent module. Add this input separately if you want to modify it.");
+			CanEditResults.bResult = false;
+			CanEditResults.UserFeedback = LOCTEXT("ModuleCanOnlyBeEditedDirectly", "This input can not be modified as it is inherent part of its parent module. Add this input separately if you want to modify it.");
 		}
 	}
 	
@@ -817,10 +809,10 @@ bool FNiagaraHierarchySummaryCategoryViewModel::IsFromBaseEmitter() const
 	return IsFromBaseEmitterCache.GetValue();
 }
 
-FHierarchyElementViewModel::FCanPerformActionResults FNiagaraHierarchySummaryCategoryViewModel::IsEditableByUser()
+FHierarchyElementViewModel::FResultWithUserFeedback FNiagaraHierarchySummaryCategoryViewModel::IsEditableByUser()
 {
-	FCanPerformActionResults CanEditResults(IsFromBaseEmitter() == false);
-	CanEditResults.CanPerformMessage = CanEditResults.bCanPerform == false ? LOCTEXT("CategoryIsFromBaseEmitter", "This category was added in the parent emitter and can not be edited. You can add new items.") : FText::GetEmpty();
+	FResultWithUserFeedback CanEditResults(IsFromBaseEmitter() == false);
+	CanEditResults.UserFeedback = CanEditResults.bResult == false ? LOCTEXT("CategoryIsFromBaseEmitter", "This category was added in the parent emitter and can not be edited. You can add new items.") : FText::GetEmpty();
 	return CanEditResults;
 }
 
@@ -842,10 +834,10 @@ bool FNiagaraHierarchyPropertyViewModel::DoesExternalDataStillExist(const UHiera
 	return PropertyObjectMap.Contains(GetData()->GetPersistentIdentity().Guids[0]);
 }
 
-FHierarchyElementViewModel::FCanPerformActionResults FNiagaraHierarchyPropertyViewModel::IsEditableByUser()
+FHierarchyElementViewModel::FResultWithUserFeedback FNiagaraHierarchyPropertyViewModel::IsEditableByUser()
 {
-	FCanPerformActionResults CanEditResults(IsFromBaseEmitter() == false);
-	CanEditResults.CanPerformMessage = CanEditResults.bCanPerform == false ? LOCTEXT("ObjectPropertyIsFromBaseEmitter", "This property was added in the parent emitter and can not be edited.") : FText::GetEmpty();
+	FResultWithUserFeedback CanEditResults(IsFromBaseEmitter() == false);
+	CanEditResults.UserFeedback = CanEditResults.bResult == false ? LOCTEXT("ObjectPropertyIsFromBaseEmitter", "This property was added in the parent emitter and can not be edited.") : FText::GetEmpty();
 	return CanEditResults;
 }
 
@@ -903,10 +895,10 @@ void FNiagaraHierarchyRendererViewModel::RefreshChildrenDataInternal()
 	GetDataMutable()->GetChildrenMutable().Append(NewChildren);
 }
 
-FHierarchyElementViewModel::FCanPerformActionResults FNiagaraHierarchyRendererViewModel::IsEditableByUser()
+FHierarchyElementViewModel::FResultWithUserFeedback FNiagaraHierarchyRendererViewModel::IsEditableByUser()
 {
-	FCanPerformActionResults CanEditResults(IsFromBaseEmitter() == false);
-	CanEditResults.CanPerformMessage = CanEditResults.bCanPerform == false ? LOCTEXT("RendererIsFromBaseEmitter", "This renderer was added in the parent emitter and can not be edited.") : FText::GetEmpty();
+	FResultWithUserFeedback CanEditResults(IsFromBaseEmitter() == false);
+	CanEditResults.UserFeedback = CanEditResults.bResult == false ? LOCTEXT("RendererIsFromBaseEmitter", "This renderer was added in the parent emitter and can not be edited.") : FText::GetEmpty();
 	return CanEditResults;
 }
 
@@ -931,10 +923,10 @@ bool FNiagaraHierarchyEmitterPropertiesViewModel::IsFromBaseEmitter() const
 	return IsFromBaseEmitterCache.GetValue();
 }
 
-FHierarchyElementViewModel::FCanPerformActionResults FNiagaraHierarchyEmitterPropertiesViewModel::IsEditableByUser()
+FHierarchyElementViewModel::FResultWithUserFeedback FNiagaraHierarchyEmitterPropertiesViewModel::IsEditableByUser()
 {
-	FCanPerformActionResults CanEditResults(IsFromBaseEmitter() == false);
-	CanEditResults.CanPerformMessage = CanEditResults.bCanPerform == false ? LOCTEXT("EmitterPropertiesIsFromBaseEmitter", "These emitter properties were added in the parent emitter and can not be edited.") : FText::GetEmpty();
+	FResultWithUserFeedback CanEditResults(IsFromBaseEmitter() == false);
+	CanEditResults.UserFeedback = CanEditResults.bResult == false ? LOCTEXT("EmitterPropertiesIsFromBaseEmitter", "These emitter properties were added in the parent emitter and can not be edited.") : FText::GetEmpty();
 	return CanEditResults;
 }
 
@@ -1038,10 +1030,10 @@ void FNiagaraHierarchyEventHandlerViewModel::RefreshChildrenDataInternal()
 	GetDataMutable()->GetChildrenMutable().Append(NewChildren);
 }
 
-FHierarchyElementViewModel::FCanPerformActionResults FNiagaraHierarchyEventHandlerViewModel::IsEditableByUser()
+FHierarchyElementViewModel::FResultWithUserFeedback FNiagaraHierarchyEventHandlerViewModel::IsEditableByUser()
 {
-	FCanPerformActionResults CanEditResults(IsFromBaseEmitter() == false);
-	CanEditResults.CanPerformMessage = CanEditResults.bCanPerform == false ? LOCTEXT("EventHandlerIsFromBaseEmitter", "This event handler was added in the parent emitter and can not be edited.") : FText::GetEmpty();
+	FResultWithUserFeedback CanEditResults(IsFromBaseEmitter() == false);
+	CanEditResults.UserFeedback = CanEditResults.bResult == false ? LOCTEXT("EventHandlerIsFromBaseEmitter", "This event handler was added in the parent emitter and can not be edited.") : FText::GetEmpty();
 	return CanEditResults;
 }
 
@@ -1150,10 +1142,10 @@ void FNiagaraHierarchyEventHandlerPropertiesViewModel::RefreshChildrenDataIntern
 	GetDataMutable()->GetChildrenMutable().Append(NewPropertiesChildren);
 }
 
-FHierarchyElementViewModel::FCanPerformActionResults FNiagaraHierarchyEventHandlerPropertiesViewModel::IsEditableByUser()
+FHierarchyElementViewModel::FResultWithUserFeedback FNiagaraHierarchyEventHandlerPropertiesViewModel::IsEditableByUser()
 {
-	FCanPerformActionResults CanEditResults(IsFromBaseEmitter() == false);
-	CanEditResults.CanPerformMessage = CanEditResults.bCanPerform == false ? LOCTEXT("EventHandlerPropertiesIsFromBaseEmitter", "This property item was added in the parent emitter and can not be edited.") : FText::GetEmpty();
+	FResultWithUserFeedback CanEditResults(IsFromBaseEmitter() == false);
+	CanEditResults.UserFeedback = CanEditResults.bResult == false ? LOCTEXT("EventHandlerPropertiesIsFromBaseEmitter", "This property item was added in the parent emitter and can not be edited.") : FText::GetEmpty();
 	return CanEditResults;
 }
 
@@ -1253,10 +1245,10 @@ void FNiagaraHierarchySimStageViewModel::RefreshChildrenDataInternal()
 	GetDataMutable()->GetChildrenMutable().Append(NewChildren);
 }
 
-FHierarchyElementViewModel::FCanPerformActionResults FNiagaraHierarchySimStageViewModel::IsEditableByUser()
+FHierarchyElementViewModel::FResultWithUserFeedback FNiagaraHierarchySimStageViewModel::IsEditableByUser()
 {
-	FCanPerformActionResults CanEditResults(IsFromBaseEmitter() == false);
-	CanEditResults.CanPerformMessage = CanEditResults.bCanPerform == false ? LOCTEXT("SimStageIsFromBaseEmitter", "This simulation stage was added in the parent emitter and can not be edited.") : FText::GetEmpty();
+	FResultWithUserFeedback CanEditResults(IsFromBaseEmitter() == false);
+	CanEditResults.UserFeedback = CanEditResults.bResult == false ? LOCTEXT("SimStageIsFromBaseEmitter", "This simulation stage was added in the parent emitter and can not be edited.") : FText::GetEmpty();
 	return CanEditResults;
 }
 
@@ -1339,10 +1331,10 @@ void FNiagaraHierarchySimStagePropertiesViewModel::RefreshChildrenDataInternal()
 	GetDataMutable()->GetChildrenMutable().Append(NewPropertiesChildren);
 }
 
-FHierarchyElementViewModel::FCanPerformActionResults FNiagaraHierarchySimStagePropertiesViewModel::IsEditableByUser()
+FHierarchyElementViewModel::FResultWithUserFeedback FNiagaraHierarchySimStagePropertiesViewModel::IsEditableByUser()
 {
-	FCanPerformActionResults CanEditResults(IsFromBaseEmitter() == false);
-	CanEditResults.CanPerformMessage = CanEditResults.bCanPerform == false ? LOCTEXT("RendererPropertiesIsFromBaseEmitter", "This renderer's properties were added in the parent emitter and can not be edited.") : FText::GetEmpty();
+	FResultWithUserFeedback CanEditResults(IsFromBaseEmitter() == false);
+	CanEditResults.UserFeedback = CanEditResults.bResult == false ? LOCTEXT("RendererPropertiesIsFromBaseEmitter", "This renderer's properties were added in the parent emitter and can not be edited.") : FText::GetEmpty();
 	return CanEditResults;
 }
 
diff --git a/Engine/Plugins/FX/Niagara/Source/NiagaraEditor/Private/ViewModels/HierarchyEditor/NiagaraUserParametersHierarchyViewModel.cpp b/Engine/Plugins/FX/Niagara/Source/NiagaraEditor/Private/ViewModels/HierarchyEditor/NiagaraUserParametersHierarchyViewModel.cpp
--- a/Engine/Plugins/FX/Niagara/Source/NiagaraEditor/Private/ViewModels/HierarchyEditor/NiagaraUserParametersHierarchyViewModel.cpp
+++ b/Engine/Plugins/FX/Niagara/Source/NiagaraEditor/Private/ViewModels/HierarchyEditor/NiagaraUserParametersHierarchyViewModel.cpp
@@ -19,7 +19,7 @@ void UNiagaraHierarchyUserParameter::Initialize(UNiagaraScriptVariable& InUserPa
 	SetIdentity(FHierarchyElementIdentity({InUserParameterScriptVariable.Metadata.GetVariableGuid()}, {}));
 }
 
-UObject* FNiagaraHierarchyUserParameterViewModel::GetDataForEditing()
+UObject* FNiagaraHierarchyUserParameterViewModel::GetObjectForEditing()
 {
 	UNiagaraHierarchyUserParameter* HierarchyUserParameter = Cast<UNiagaraHierarchyUserParameter>(GetDataMutable<UNiagaraHierarchyUserParameter>());
 	FNiagaraVariable ContainedVariable = HierarchyUserParameter->GetUserParameter();
diff --git a/Engine/Plugins/FX/Niagara/Source/NiagaraEditor/Public/ViewModels/HierarchyEditor/NiagaraScriptParametersHierarchyViewModel.h b/Engine/Plugins/FX/Niagara/Source/NiagaraEditor/Public/ViewModels/HierarchyEditor/NiagaraScriptParametersHierarchyViewModel.h
--- a/Engine/Plugins/FX/Niagara/Source/NiagaraEditor/Public/ViewModels/HierarchyEditor/NiagaraScriptParametersHierarchyViewModel.h
+++ b/Engine/Plugins/FX/Niagara/Source/NiagaraEditor/Public/ViewModels/HierarchyEditor/NiagaraScriptParametersHierarchyViewModel.h
@@ -100,7 +100,7 @@ struct FNiagaraHierarchyScriptCategoryViewModel : public FHierarchyCategoryViewM
 protected:
 	// We don't sort anything to allow free authoring of UI
 	virtual void SortChildrenData() const override {}
-	virtual FCanPerformActionResults CanDropOnInternal(TSharedPtr<FHierarchyElementViewModel> DraggedElement, EItemDropZone ItemDropZone) override;
+	virtual FResultWithUserFeedback CanDropOnInternal(TSharedPtr<FHierarchyElementViewModel> DraggedElement, EItemDropZone ItemDropZone) override;
 };
 
 struct FNiagaraHierarchyScriptParameterViewModel : public FHierarchyItemViewModel
@@ -111,15 +111,15 @@ struct FNiagaraHierarchyScriptParameterViewModel : public FHierarchyItemViewMode
 protected:
 	virtual bool RepresentsExternalData() const override { return true; }
 	virtual bool DoesExternalDataStillExist(const UHierarchyDataRefreshContext* Context) const override;
-	virtual UObject* GetDataForEditing() override { return Cast<UNiagaraHierarchyScriptParameter>(GetDataMutable())->GetScriptVariable(); }
+	virtual UObject* GetObjectForEditing() override { return Cast<UNiagaraHierarchyScriptParameter>(GetDataMutable())->GetScriptVariable(); }
 	/** We want to be able to edit in the details panel regardless of source or hierarchy item. */
 	virtual bool AllowEditingInDetailsPanel() const override { return true; }
 	virtual bool CanRenameInternal() override { return false; }
+	virtual FResultWithUserFeedback CanContainInternal(TSubclassOf<UHierarchyElement> InHierarchyElementType) override;
 	
-	virtual FCanPerformActionResults CanDropOnInternal(TSharedPtr<FHierarchyElementViewModel>, EItemDropZone ItemDropZone) override;
-	virtual void OnDroppedOnInternal(TSharedPtr<FHierarchyElementViewModel> DroppedItem, EItemDropZone ItemDropZone) override;
+	virtual FResultWithUserFeedback CanDropOnInternal(TSharedPtr<FHierarchyElementViewModel>, EItemDropZone ItemDropZone) override;
 	
-	virtual bool CanHaveChildren() const override { return bIsForHierarchy == true; }
+	virtual FResultWithUserFeedback CanHaveChildren() const override { return bIsForHierarchy == true; }
 };
 
 struct FNiagaraHierarchyScriptRootViewModel : public FHierarchyRootViewModel
diff --git a/Engine/Plugins/FX/Niagara/Source/NiagaraEditor/Public/ViewModels/HierarchyEditor/NiagaraSummaryViewViewModel.h b/Engine/Plugins/FX/Niagara/Source/NiagaraEditor/Public/ViewModels/HierarchyEditor/NiagaraSummaryViewViewModel.h
--- a/Engine/Plugins/FX/Niagara/Source/NiagaraEditor/Public/ViewModels/HierarchyEditor/NiagaraSummaryViewViewModel.h
+++ b/Engine/Plugins/FX/Niagara/Source/NiagaraEditor/Public/ViewModels/HierarchyEditor/NiagaraSummaryViewViewModel.h
@@ -202,10 +202,10 @@ private:
 	
 	virtual FString ToString() const override;
 
-	virtual FCanPerformActionResults IsEditableByUser() override;
-	virtual bool CanHaveChildren() const override { return bIsForHierarchy == false; }
+	virtual FResultWithUserFeedback IsEditableByUser() override;
+	virtual FResultWithUserFeedback CanHaveChildren() const override { return bIsForHierarchy == false; }
 	virtual bool CanRenameInternal() override { return false; }
-	virtual FCanPerformActionResults CanDropOnInternal(TSharedPtr<FHierarchyElementViewModel>, EItemDropZone ItemDropZone) override;
+	virtual FResultWithUserFeedback CanDropOnInternal(TSharedPtr<FHierarchyElementViewModel>, EItemDropZone ItemDropZone) override;
 	virtual const UHierarchySection* GetSectionInternal() const override;
 	
 	virtual bool RepresentsExternalData() const override { return true; }
@@ -252,14 +252,15 @@ protected:
 
 	void RefreshChildDynamicInputs(bool bClearCache = false);
 	
-	virtual bool CanHaveChildren() const override;
-	virtual FCanPerformActionResults IsEditableByUser() override;
+	virtual FResultWithUserFeedback CanHaveChildren() const override;
+	virtual FResultWithUserFeedback CanContainInternal(TSubclassOf<UHierarchyElement> InHierarchyElementType) override;
+	virtual FResultWithUserFeedback IsEditableByUser() override;
 	virtual bool RepresentsExternalData() const override { return true; }
 	virtual bool DoesExternalDataStillExist(const UHierarchyDataRefreshContext* Context) const override { ClearCache(); return GetInputData().IsSet(); }
 	virtual void RefreshChildrenDataInternal() override;
-
-	virtual FCanPerformActionResults CanDropOnInternal(TSharedPtr<FHierarchyElementViewModel>, EItemDropZone ItemDropZone) override;
-	virtual void OnDroppedOnInternal(TSharedPtr<FHierarchyElementViewModel> DroppedItem, EItemDropZone ItemDropZone) override;
+	
+	virtual FResultWithUserFeedback CanDropOnInternal(TSharedPtr<FHierarchyElementViewModel>, EItemDropZone ItemDropZone) override;
+	virtual void PostAddFixup() override;
 
 	virtual void AppendDynamicContextMenuForSingleElement(UToolMenu* ToolMenu) override;
 private:
@@ -277,12 +278,12 @@ struct FNiagaraAssignmentInputViewModel : public FHierarchyItemViewModel
 	FNiagaraAssignmentInputViewModel(UNiagaraHierarchyAssignmentInput* ModuleInput, TSharedRef<FHierarchyElementViewModel> InParent, TWeakObjectPtr<UNiagaraSummaryViewViewModel> ViewModel)
 	: FHierarchyItemViewModel(ModuleInput, InParent, ViewModel) {}
 
-	virtual FCanPerformActionResults CanDropOnInternal(TSharedPtr<FHierarchyElementViewModel>, EItemDropZone ItemDropZone) override;
+	virtual FResultWithUserFeedback CanDropOnInternal(TSharedPtr<FHierarchyElementViewModel>, EItemDropZone ItemDropZone) override;
 	
 	TWeakObjectPtr<UNiagaraNodeAssignment> GetAssignmentNode() const;
 	TOptional<FNiagaraStackGraphUtilities::FMatchingFunctionInputData> GetInputData() const;
 
-	virtual bool CanHaveChildren() const override { return false; }
+	virtual FResultWithUserFeedback CanHaveChildren() const override { return false; }
 
 	virtual FString ToString() const override;
 	virtual TArray<FString> GetSearchTerms() const override;
@@ -291,7 +292,7 @@ struct FNiagaraAssignmentInputViewModel : public FHierarchyItemViewModel
 
 	void ClearCache() const;
 protected:
-	virtual FCanPerformActionResults IsEditableByUser() override;
+	virtual FResultWithUserFeedback IsEditableByUser() override;
 	virtual bool RepresentsExternalData() const override { return true; }
 	virtual bool DoesExternalDataStillExist(const UHierarchyDataRefreshContext* Context) const override { ClearCache(); return GetInputData().IsSet(); }
 private:
@@ -310,7 +311,7 @@ struct FNiagaraHierarchySummaryCategoryViewModel : public FHierarchyCategoryView
 	bool IsFromBaseEmitter() const;
 	
 protected:
-	virtual FCanPerformActionResults IsEditableByUser() override;
+	virtual FResultWithUserFeedback IsEditableByUser() override;
 	
 	mutable TOptional<bool> IsFromBaseEmitterCache;
 };
@@ -327,7 +328,7 @@ protected:
 	virtual bool RepresentsExternalData() const override { return true; }
 	virtual bool DoesExternalDataStillExist(const UHierarchyDataRefreshContext* Context) const override;
 	
-	virtual FCanPerformActionResults IsEditableByUser() override;
+	virtual FResultWithUserFeedback IsEditableByUser() override;
 	
 	mutable TOptional<bool> IsFromBaseEmitterCache;	
 };
@@ -348,8 +349,8 @@ struct FNiagaraHierarchyRendererViewModel : public FHierarchyItemViewModel
 protected:
 	virtual void RefreshChildrenDataInternal() override;
 	
-	virtual FCanPerformActionResults IsEditableByUser() override;
-	virtual bool CanHaveChildren() const override { return bIsForHierarchy == false; }
+	virtual FResultWithUserFeedback IsEditableByUser() override;
+	virtual FResultWithUserFeedback CanHaveChildren() const override { return bIsForHierarchy == false; }
 	
 	virtual bool RepresentsExternalData() const override { return true; }
 	virtual bool DoesExternalDataStillExist(const UHierarchyDataRefreshContext* Context) const override { return GetRendererProperties() != nullptr;}
@@ -373,8 +374,8 @@ struct FNiagaraHierarchyEmitterPropertiesViewModel : public FHierarchyItemViewMo
 
 	bool IsFromBaseEmitter() const;
 protected:	
-	virtual FCanPerformActionResults IsEditableByUser() override;
-	virtual bool CanHaveChildren() const override { return bIsForHierarchy == false; }
+	virtual FResultWithUserFeedback IsEditableByUser() override;
+	virtual FResultWithUserFeedback CanHaveChildren() const override { return bIsForHierarchy == false; }
 	
 	virtual bool RepresentsExternalData() const override { return true; }
 	virtual bool DoesExternalDataStillExist(const UHierarchyDataRefreshContext* Context) const override { return true; }
@@ -401,8 +402,8 @@ struct FNiagaraHierarchyEventHandlerViewModel : public FHierarchyItemViewModel
 protected:
 	virtual void RefreshChildrenDataInternal() override;
 	
-	virtual FCanPerformActionResults IsEditableByUser() override;
-	virtual bool CanHaveChildren() const override { return bIsForHierarchy == false; }
+	virtual FResultWithUserFeedback IsEditableByUser() override;
+	virtual FResultWithUserFeedback CanHaveChildren() const override { return bIsForHierarchy == false; }
 	
 	virtual bool RepresentsExternalData() const override { return true; }
 	virtual bool DoesExternalDataStillExist(const UHierarchyDataRefreshContext* Context) const override { return GetEventScriptProperties() != nullptr;}
@@ -427,8 +428,8 @@ struct FNiagaraHierarchyEventHandlerPropertiesViewModel : public FHierarchyItemV
 protected:
 	virtual void RefreshChildrenDataInternal() override;
 
-	virtual FCanPerformActionResults IsEditableByUser() override;
-	virtual bool CanHaveChildren() const override { return bIsForHierarchy == false; }
+	virtual FResultWithUserFeedback IsEditableByUser() override;
+	virtual FResultWithUserFeedback CanHaveChildren() const override { return bIsForHierarchy == false; }
 	
 	virtual bool RepresentsExternalData() const override { return true; }
 	virtual bool DoesExternalDataStillExist(const UHierarchyDataRefreshContext* Context) const override { return GetEventScriptProperties() != nullptr;}
@@ -452,8 +453,8 @@ struct FNiagaraHierarchySimStageViewModel : public FHierarchyItemViewModel
 protected:
 	virtual void RefreshChildrenDataInternal() override;
 	
-	virtual FCanPerformActionResults IsEditableByUser() override;
-	virtual bool CanHaveChildren() const override { return bIsForHierarchy == false; }
+	virtual FResultWithUserFeedback IsEditableByUser() override;
+	virtual FResultWithUserFeedback CanHaveChildren() const override { return bIsForHierarchy == false; }
 	
 	virtual bool RepresentsExternalData() const override { return true; }
 	virtual bool DoesExternalDataStillExist(const UHierarchyDataRefreshContext* Context) const override { return GetSimStage() != nullptr;}
@@ -478,8 +479,8 @@ struct FNiagaraHierarchySimStagePropertiesViewModel : public FHierarchyItemViewM
 protected:
 	virtual void RefreshChildrenDataInternal() override;
 
-	virtual FCanPerformActionResults IsEditableByUser() override;
-	virtual bool CanHaveChildren() const override { return bIsForHierarchy == false; }
+	virtual FResultWithUserFeedback IsEditableByUser() override;
+	virtual FResultWithUserFeedback CanHaveChildren() const override { return bIsForHierarchy == false; }
 	
 	virtual bool RepresentsExternalData() const override { return true; }
 	virtual bool DoesExternalDataStillExist(const UHierarchyDataRefreshContext* Context) const override { return GetSimStage() != nullptr;}
diff --git a/Engine/Plugins/FX/Niagara/Source/NiagaraEditor/Public/ViewModels/HierarchyEditor/NiagaraUserParametersHierarchyViewModel.h b/Engine/Plugins/FX/Niagara/Source/NiagaraEditor/Public/ViewModels/HierarchyEditor/NiagaraUserParametersHierarchyViewModel.h
--- a/Engine/Plugins/FX/Niagara/Source/NiagaraEditor/Public/ViewModels/HierarchyEditor/NiagaraUserParametersHierarchyViewModel.h
+++ b/Engine/Plugins/FX/Niagara/Source/NiagaraEditor/Public/ViewModels/HierarchyEditor/NiagaraUserParametersHierarchyViewModel.h
@@ -48,7 +48,7 @@ struct FNiagaraHierarchyUserParameterViewModel : public FHierarchyItemViewModel
 	: FHierarchyItemViewModel(InItem, InParent, InHierarchyViewModel) {}
 	
 	/** For editing in the details panel we want to handle the script variable that is represented by the hierarchy item, not the hierarchy item itself. */
-	virtual UObject* GetDataForEditing() override;
+	virtual UObject* GetObjectForEditing() override;
 	/** We want to be able to edit in the details panel regardless of source or hierarchy item. */
 	virtual bool AllowEditingInDetailsPanel() const override { return true; }
 
diff --git a/Engine/Source/Editor/DataHierarchyEditor/DataHierarchyEditor.Build.cs b/Engine/Source/Editor/DataHierarchyEditor/DataHierarchyEditor.Build.cs
--- a/Engine/Source/Editor/DataHierarchyEditor/DataHierarchyEditor.Build.cs
+++ b/Engine/Source/Editor/DataHierarchyEditor/DataHierarchyEditor.Build.cs
@@ -25,6 +25,7 @@ public class DataHierarchyEditor : ModuleRules
                 "Slate",
                 "UnrealEd",
                 "ToolMenus",
+				"ToolWidgets",
                 "InputCore"
             }
         );
diff --git a/Engine/Source/Editor/DataHierarchyEditor/Internal/DataHierarchyEditorMisc.h b/Engine/Source/Editor/DataHierarchyEditor/Internal/DataHierarchyEditorMisc.h
--- a/Engine/Source/Editor/DataHierarchyEditor/Internal/DataHierarchyEditorMisc.h
+++ b/Engine/Source/Editor/DataHierarchyEditor/Internal/DataHierarchyEditorMisc.h
@@ -4,9 +4,20 @@
 
 #include "Containers/Set.h"
 #include "UObject/NameTypes.h"
+#include "DataHierarchyEditorMisc.generated.h"
 
+class SDataHierarchyEditor;
 
 namespace UE::DataHierarchyEditor
 {
 	DATAHIERARCHYEDITOR_API FName GetUniqueName(FName CandidateName, const TSet<FName>& ExistingNames);
 }
+
+UCLASS()
+class UDataHierarchyEditorMenuContext : public UObject
+{
+	GENERATED_BODY()
+
+public:
+	TWeakPtr<SDataHierarchyEditor> Widget;
+};
\ No newline at end of file
diff --git a/Engine/Source/Editor/DataHierarchyEditor/Internal/Widgets/SDataHierarchyEditor.h b/Engine/Source/Editor/DataHierarchyEditor/Internal/Widgets/SDataHierarchyEditor.h
--- a/Engine/Source/Editor/DataHierarchyEditor/Internal/Widgets/SDataHierarchyEditor.h
+++ b/Engine/Source/Editor/DataHierarchyEditor/Internal/Widgets/SDataHierarchyEditor.h
@@ -27,9 +27,6 @@ struct FHierarchySectionViewModel;
 class SHierarchySection : public SCompoundWidget
 {
 	DECLARE_DELEGATE_OneParam(FOnSectionActivated, TSharedPtr<FHierarchySectionViewModel> SectionViewModel)
-
-	static UE_API const FSlateBrush LeftDropBrush;
-	static UE_API const FSlateBrush RightDropBrush;
 	
 	SLATE_BEGIN_ARGS(SHierarchySection)
 		{}
@@ -106,14 +103,10 @@ class SDataHierarchyEditor : public SCompoundWidget, public FNotifyHook
 
 	SLATE_BEGIN_ARGS(SDataHierarchyEditor)
 		: _bReadOnly(true)
-		, _ItemRowStyle(&FAppStyle::Get().GetWidgetStyle<FTableRowStyle>("TableView.Row"))
-		, _CategoryRowStyle(&FDataHierarchyEditorStyle::Get().GetWidgetStyle<FTableRowStyle>("HierarchyEditor.Row.Category"))
 	{}
 		SLATE_ARGUMENT(bool, bReadOnly)
 		SLATE_EVENT(FOnGenerateRowContentWidget, OnGenerateRowContentWidget)
 		SLATE_EVENT(FOnGenerateCustomDetailsPanelNameWidget, OnGenerateCustomDetailsPanelNameWidget)
-		SLATE_STYLE_ARGUMENT(FTableRowStyle, ItemRowStyle)
-		SLATE_STYLE_ARGUMENT(FTableRowStyle, CategoryRowStyle)
 	SLATE_END_ARGS()
 
 public:
@@ -178,7 +171,9 @@ private:
 	UE_API TSharedPtr<FHierarchySectionViewModel> GetActiveSourceSection() const;
 	UE_API UHierarchySection* GetActiveSourceSectionData() const;
 	
-	UE_API void OnSelectionChanged(TSharedPtr<FHierarchyElementViewModel> HierarchyItem, ESelectInfo::Type Type, bool bFromHierarchy) const;
+	TSharedRef<SWidget> OnGetCustomAddContent();
+
+	void OnSelectionChanged(TSharedPtr<FHierarchyElementViewModel> SelectedHierarchyElement, ESelectInfo::Type Type) const;
 
 	UE_API void RunSourceSearch();
 	UE_API void OnSourceSearchTextChanged(const FText& Text);
@@ -218,7 +213,7 @@ private:
 	UE_API void SelectPreviousSourceSearchResult();
 	UE_API TOptional<SSearchBox::FSearchResultData> GetSearchResultData() const;
 	
-	UE_API FHierarchyElementViewModel::FCanPerformActionResults CanDropOnRoot(TSharedPtr<FHierarchyElementViewModel> DraggedItem) const;
+	FHierarchyElementViewModel::FResultWithUserFeedback CanDropOnRoot(TSharedPtr<FHierarchyElementViewModel> DraggedItem) const;
 
 	/** Callback functions for the root widget */
 	UE_API FReply HandleHierarchyRootDrop(const FGeometry&, const FDragDropEvent& DragDropEvent) const;
@@ -229,13 +224,11 @@ private:
 
 	UE_API virtual void NotifyPostChange(const FPropertyChangedEvent& PropertyChangedEvent, FProperty* PropertyThatChanged) override;
 private:
+	/** A weak pointer to the hierarchy view model.
+	 * An external view model created and managed by the implementer. */
 	TWeakObjectPtr<UDataHierarchyViewModelBase> HierarchyViewModel;
-	
 	TArray<FSearchItem> SourceSearchResults;
 	TOptional<FSearchItem> FocusedSearchResult;
-
-	const FTableRowStyle* CategoryRowStyle = nullptr;
-	const FTableRowStyle* ItemRowStyle = nullptr;
 	
 	mutable TWeakPtr<FHierarchyElementViewModel> SelectedDetailsPanelItemViewModel;
 
@@ -244,6 +237,7 @@ private:
 	TSharedPtr<FHierarchyRootViewModel> SourceRootViewModel;
 	TSharedPtr<STreeView<TSharedPtr<FHierarchyElementViewModel>>> SourceTreeView;
 	TSharedPtr<STreeView<TSharedPtr<FHierarchyElementViewModel>>> HierarchyTreeView;
+	TStrongObjectPtr<UHierarchySection> AllSourceSection;
 	TSharedPtr<FHierarchySectionViewModel> DefaultSourceSectionViewModel;
 	TWeakPtr<struct FHierarchySectionViewModel> ActiveSourceSection;
 	TSharedPtr<class SWrapBox> SourceSectionBox;
@@ -258,24 +252,26 @@ private:
 	TSharedPtr<FActiveTimerHandle> RefreshSectionsViewNextFrameHandle;
 };
 
-class SHierarchyCategory : public SCompoundWidget
+class SHierarchyElement : public SCompoundWidget
 {
-	SLATE_BEGIN_ARGS(SHierarchyCategory)
+	SLATE_BEGIN_ARGS(SHierarchyElement)
+		: _Style(&FCoreStyle::Get().GetWidgetStyle<FInlineEditableTextBlockStyle>("InlineEditableTextBlockStyle") )
 	{}
+		SLATE_STYLE_ARGUMENT(FInlineEditableTextBlockStyle, Style)
 		SLATE_EVENT(FIsSelected, IsSelected)
 	SLATE_END_ARGS()
 
 public:
-	UE_API void Construct(const FArguments& InArgs, TSharedPtr<struct FHierarchyCategoryViewModel> InCategoryViewModel);
+	UE_API void Construct(const FArguments& InArgs, TSharedPtr<struct FHierarchyElementViewModel> InViewModel);
 
 	UE_API void EnterEditingMode() const;
-	UE_API bool OnVerifyCategoryRename(const FText& NewName, FText& OutTooltip) const;
+	UE_API bool OnVerifyRename(const FText& NewName, FText& OutTooltip) const;
 	
-	UE_API FText GetCategoryText() const;
-	UE_API void OnRenameCategory(const FText& NewText, ETextCommit::Type) const;
+	UE_API FText GetElementText() const;
+	UE_API void OnRenameElement(const FText& NewText, ETextCommit::Type) const;
 	
 private:
-	TWeakPtr<struct FHierarchyCategoryViewModel> CategoryViewModelWeak;
+	TWeakPtr<struct FHierarchyElementViewModel> ElementViewModelWeak;
 	TSharedPtr<SInlineEditableTextBlock> InlineEditableTextBlock;
 };
 
diff --git a/Engine/Source/Editor/DataHierarchyEditor/Private/DataHierarchyViewModelBase.cpp b/Engine/Source/Editor/DataHierarchyEditor/Private/DataHierarchyViewModelBase.cpp
--- a/Engine/Source/Editor/DataHierarchyEditor/Private/DataHierarchyViewModelBase.cpp
+++ b/Engine/Source/Editor/DataHierarchyEditor/Private/DataHierarchyViewModelBase.cpp
@@ -22,6 +22,10 @@
 
 #define LOCTEXT_NAMESPACE "DataHierarchyEditor"
 
+FName UDataHierarchyViewModelBase::AllSectionHierarchyObjectName = "AllSection_Hierarchy";
+FName UDataHierarchyViewModelBase::AllSectionSourceObjectName = "AllSection_Source";
+FText UDataHierarchyViewModelBase::AllSectionDisplayName = LOCTEXT("AllSectionDisplayName", "All");
+
 const TArray<UHierarchyElement*>& UHierarchyElement::GetChildren() const
 {
 	return Children;
@@ -122,6 +126,18 @@ TArray<FHierarchyElementIdentity> UHierarchyElement::GetParentIdentities() const
 	return ParentIdentities;
 }
 
+TArray<const TObjectPtr<UHierarchyElement>> UHierarchyElement::GetParentElements() const
+{
+	TArray<const TObjectPtr<UHierarchyElement>> Result;
+
+	for(UHierarchyElement* Parent = Cast<UHierarchyElement>(GetOuter()); Parent != nullptr; Parent = Cast<UHierarchyElement>(Parent->GetOuter()))
+	{
+		Result.Add(Parent);
+	}
+
+	return Result;
+}
+
 bool UHierarchyElement::Modify(bool bAlwaysMarkDirty)
 {
 	bool bSavedToTransactionBuffer = true;
@@ -309,9 +325,14 @@ bool FHierarchyCategoryViewModel::IsTopCategoryActive() const
 	return false;	
 }
 
-FHierarchyElementViewModel::FCanPerformActionResults FHierarchyCategoryViewModel::CanDropOnInternal(TSharedPtr<FHierarchyElementViewModel> DraggedElement, EItemDropZone ItemDropZone)
+FHierarchyElementViewModel::FResultWithUserFeedback FHierarchyCategoryViewModel::CanContainInternal(TSubclassOf<UHierarchyElement> InHierarchyElement)
+{
+	return (InHierarchyElement->IsChildOf<UHierarchyItem>() || InHierarchyElement->IsChildOf<UHierarchyCategory>());	
+}
+
+FHierarchyElementViewModel::FResultWithUserFeedback FHierarchyCategoryViewModel::CanDropOnInternal(TSharedPtr<FHierarchyElementViewModel> DraggedElement, EItemDropZone ItemDropZone)
 {
-	FCanPerformActionResults Results(false);
+	FResultWithUserFeedback Results(false);
 	
 	TArray<TSharedPtr<FHierarchyCategoryViewModel>> TargetChildrenCategories;
 	GetChildrenViewModelsForType<UHierarchyCategory, FHierarchyCategoryViewModel>(TargetChildrenCategories);
@@ -336,22 +357,24 @@ FHierarchyElementViewModel::FCanPerformActionResults FHierarchyCategoryViewModel
 
 				if(bContainsSiblingWithSameName)
 				{
-					Results.bCanPerform = false;
-					Results.CanPerformMessage = LOCTEXT("CantDropCategorNextToCategorySameSiblingNames", "A category of the same name already exists here, potentially in a different section. Please rename your category first.");
+					Results.bResult = false;
+					Results.UserFeedback = LOCTEXT("CantDropCategorNextToCategorySameSiblingNames", "A category of the same name already exists here, potentially in a different section. Please rename your category first.");
 					return Results;
 				}
 
-				Results.CanPerformMessage = LOCTEXT("MoveCategoryText", "Move category here");
+				Results.UserFeedback = LOCTEXT("MoveCategoryText", "Move category here");
 
 				// if we are making a category a sibling of another at the root level, the section will be set to the currently active section. Let that be known.
 				if(Parent.Pin()->GetData()->IsA<UHierarchyRoot>())
 				{
 					UHierarchyCategory* DraggedCategory = Cast<UHierarchyCategory>(DraggedElement->GetDataMutable());
-					if(DraggedCategory->GetSection() != HierarchyViewModel->GetActiveHierarchySectionData())
+					const UHierarchySection* CurrentElementSection = DraggedCategory->GetSection() == nullptr ? HierarchyViewModel->GetDefaultHierarchySectionViewModel()->GetData<UHierarchySection>() : DraggedCategory->GetSection();
+					
+					if(CurrentElementSection != HierarchyViewModel->GetActiveHierarchySectionData())
 					{
 						FText SectionChangeBaseText = LOCTEXT("CategorySectionWillUpdateDueToDrop", "The section of the category will change to {0} after the drop");
-						FText ActualSectionChangeText = FText::FormatOrdered(SectionChangeBaseText, HierarchyViewModel->GetActiveHierarchySectionData() == nullptr ? FText::FromString("All") : HierarchyViewModel->GetActiveHierarchySectionData()->GetSectionNameAsText());
-						Results.CanPerformMessage = FText::FormatOrdered(FText::AsCultureInvariant("{0}\n{1}"), Results.CanPerformMessage, ActualSectionChangeText);
+						FText ActualSectionChangeText = FText::FormatOrdered(SectionChangeBaseText, HierarchyViewModel->GetActiveHierarchySectionViewModel()->GetSectionNameAsText());
+						Results.UserFeedback = FText::FormatOrdered(FText::AsCultureInvariant("{0}\n{1}"), Results.UserFeedback.Get(FText::GetEmpty()), ActualSectionChangeText);
 					}
 				}
 			}
@@ -364,33 +387,33 @@ FHierarchyElementViewModel::FCanPerformActionResults FHierarchyCategoryViewModel
 
 				if(bContainsChildrenCategoriesWithSameName)
 				{
-					Results.bCanPerform = false;
-					Results.CanPerformMessage = LOCTEXT("CantDropCategoryOnCategorySameChildCategoryName", "A sub-category of the same name already exists! Please rename your category first.");
+					Results.bResult = false;
+					Results.UserFeedback = LOCTEXT("CantDropCategoryOnCategorySameChildCategoryName", "A sub-category of the same name already exists! Please rename your category first.");
 					return Results;
 				}
 
-				Results.CanPerformMessage = LOCTEXT("CreateSubcategory", "Drop category here to create a sub-category");
+				Results.UserFeedback = LOCTEXT("CreateSubcategory", "Drop category here to create a sub-category");
 			}
 
-			Results.bCanPerform = true;
+			Results.bResult = true;
 			return Results;
 		}
 		else if(DraggedElement->GetData()->IsA<UHierarchyItem>())
 		{
 			// items can generally be dropped onto categories
-			Results.bCanPerform = EItemDropZone::OntoItem == ItemDropZone;
+			Results.bResult = EItemDropZone::OntoItem == ItemDropZone;
 
-			if(Results.bCanPerform)
+			if(Results.bResult)
 			{
 				if(DraggedElement->IsForHierarchy() == false)
 				{
 					FText Message = LOCTEXT("AddItemToCategoryDragMessage", "Add {0} to {1}");
-					Results.CanPerformMessage = FText::FormatOrdered(Message, FText::FromString(DraggedElement->ToString()), FText::FromString(ToString()));
+					Results.UserFeedback = FText::FormatOrdered(Message, FText::FromString(DraggedElement->ToString()), FText::FromString(ToString()));
 				}
 				else
 				{
 					FText Message = LOCTEXT("MoveItemToCategoryDragMessage", "Move {0} to {1}");
-					Results.CanPerformMessage = FText::FormatOrdered(Message, FText::FromString(DraggedElement->ToString()), FText::FromString(ToString()));
+					Results.UserFeedback = FText::FormatOrdered(Message, FText::FromString(DraggedElement->ToString()), FText::FromString(ToString()));
 				}
 			}
 		}
@@ -487,7 +510,9 @@ void UDataHierarchyViewModelBase::Initialize()
 	HierarchyRootViewModel->AddChildFilter(FHierarchyElementViewModel::FOnFilterChild::CreateUObject(this, &UDataHierarchyViewModelBase::FilterForUncategorizedRootItemsInAllSection));
 	HierarchyRootViewModel->SyncViewModelsToData();
 
-	DefaultHierarchySectionViewModel = MakeShared<FHierarchySectionViewModel>(nullptr, GetHierarchyRootViewModel().ToSharedRef(), this);
+	AllSection = NewObject<UHierarchySection>(this, AllSectionHierarchyObjectName, RF_Transient);
+	DefaultHierarchySectionViewModel = MakeShared<FHierarchySectionViewModel>(AllSection, GetHierarchyRootViewModel().ToSharedRef(), this);
+	DefaultHierarchySectionViewModel->SetSectionName(FName(AllSectionDisplayName.ToString()));
 	SetActiveHierarchySection(DefaultHierarchySectionViewModel);
 	
 	InitializeInternal();
@@ -623,16 +648,16 @@ bool UDataHierarchyViewModelBase::MatchesContext(const FTransactionContext& InCo
 
 bool UDataHierarchyViewModelBase::FilterForHierarchySection(TSharedPtr<const FHierarchyElementViewModel> ViewModel) const
 {
-	if(ActiveHierarchySection.IsValid())
+	if(ActiveHierarchySectionViewModel.IsValid())
 	{
 		// If the currently selected section data is nullptr, it's the All section, and we let everything pass
-		if(ActiveHierarchySection.Pin()->GetData() == nullptr)
+		if(ActiveHierarchySectionViewModel.Pin()->IsAllSection())
 		{
 			return true;
 		}
 
 		// if not, we check against identical section data
-		return ActiveHierarchySection.Pin()->GetData() == ViewModel->GetSection();
+		return ActiveHierarchySectionViewModel.Pin()->GetData() == ViewModel->GetSection();
 	}
 
 	return true;
@@ -640,10 +665,10 @@ bool UDataHierarchyViewModelBase::FilterForHierarchySection(TSharedPtr<const FHi
 
 bool UDataHierarchyViewModelBase::FilterForUncategorizedRootItemsInAllSection(TSharedPtr<const FHierarchyElementViewModel> ViewModel) const
 {
-	if(ActiveHierarchySection.IsValid())
+	if(ActiveHierarchySectionViewModel.IsValid())
 	{
 		// we want to filter out all items that are directly added to the root if we aren't in the 'All' section
-		if(ActiveHierarchySection.Pin()->GetData() == nullptr)
+		if(ActiveHierarchySectionViewModel.Pin()->IsAllSection())
 		{
 			return true;
 		}
@@ -755,26 +780,24 @@ void UDataHierarchyViewModelBase::OnGetChildren(TSharedPtr<FHierarchyElementView
 
 void UDataHierarchyViewModelBase::SetActiveHierarchySection(TSharedPtr<FHierarchySectionViewModel> Section)
 {
-	ActiveHierarchySection = Section;	
+	ActiveHierarchySectionViewModel = Section;	
 	RefreshHierarchyView(true);
 	OnHierarchySectionActivatedDelegate.ExecuteIfBound(Section);
 }
 
-TSharedPtr<FHierarchySectionViewModel> UDataHierarchyViewModelBase::GetActiveHierarchySection() const
+TSharedPtr<FHierarchySectionViewModel> UDataHierarchyViewModelBase::GetActiveHierarchySectionViewModel() const
 {
-	return ActiveHierarchySection.Pin();
+	return ActiveHierarchySectionViewModel.Pin();
 }
 
 UHierarchySection* UDataHierarchyViewModelBase::GetActiveHierarchySectionData() const
 {
-	TSharedPtr<FHierarchySectionViewModel> ActiveSection = GetActiveHierarchySection();
-	return ActiveSection.IsValid() ? ActiveSection->GetDataMutable<UHierarchySection>() : nullptr;
+	return ActiveHierarchySectionViewModel.IsValid() ? ActiveHierarchySectionViewModel.Pin()->GetDataMutable<UHierarchySection>() : nullptr;
 }
 
 bool UDataHierarchyViewModelBase::IsHierarchySectionActive(const UHierarchySection* Section) const
 {
-	TSharedPtr<FHierarchySectionViewModel> ActiveSection = GetActiveHierarchySection();
-	return ActiveSection.IsValid() ? ActiveSection->GetData() == Section : false;
+	return ActiveHierarchySectionViewModel.IsValid() ? ActiveHierarchySectionViewModel.Pin()->GetData() == Section : false;
 }
 
 FString UDataHierarchyViewModelBase::OnElementToStringDebug(TSharedPtr<FHierarchyElementViewModel> ElementViewModel) const
@@ -898,7 +921,7 @@ const TArray<TSharedPtr<FHierarchyElementViewModel>>& FHierarchyElementViewModel
 {
 	FilteredChildren.Empty();
 
-	if(CanHaveChildren())
+	if(CanHaveChildren().bResult)
 	{
 		for(TSharedPtr<FHierarchyElementViewModel> Child : Children)
 		{
@@ -1089,9 +1112,12 @@ TOptional<EItemDropZone> FHierarchyElementViewModel::OnCanRowAcceptDrop(const FD
 {
 	if(TSharedPtr<FHierarchyDragDropOp> DragDropOp = DragDropEvent.GetOperationAs<FHierarchyDragDropOp>())
 	{
-		FCanPerformActionResults Results = CanDropOn(DragDropOp->GetDraggedElement().Pin(), ItemDropZone);
-		DragDropOp->SetDescription(Results.CanPerformMessage);
-		return Results.bCanPerform ? ItemDropZone : TOptional<EItemDropZone>();
+		FResultWithUserFeedback Results = CanDropOn(DragDropOp->GetDraggedElement().Pin(), ItemDropZone);
+		if(Results.UserFeedback.IsSet())
+		{
+			DragDropOp->SetDescription(Results.UserFeedback.GetValue());
+		}
+		return Results.bResult ? ItemDropZone : TOptional<EItemDropZone>();
 	}
 
 	return TOptional<EItemDropZone>();
@@ -1116,10 +1142,10 @@ void FHierarchyElementViewModel::OnRowDragLeave(const FDragDropEvent& DragDropEv
 	}
 }
 
-FHierarchyElementViewModel::FCanPerformActionResults FHierarchyElementViewModel::CanDrag()
+FHierarchyElementViewModel::FResultWithUserFeedback FHierarchyElementViewModel::CanDrag()
 {
-	FCanPerformActionResults Results = IsEditableByUser();
-	if(Results.bCanPerform == false)
+	FResultWithUserFeedback Results = IsEditableByUser();
+	if(Results.bResult == false)
 	{
 		return Results;
 	}
@@ -1127,9 +1153,143 @@ FHierarchyElementViewModel::FCanPerformActionResults FHierarchyElementViewModel:
 	return CanDragInternal();
 }
 
-FHierarchyElementViewModel::FCanPerformActionResults FHierarchyElementViewModel::CanDropOnInternal(TSharedPtr<FHierarchyElementViewModel>, EItemDropZone ItemDropZone)
+FHierarchyElementViewModel::FResultWithUserFeedback FHierarchyElementViewModel::CanDropOn(TSharedPtr<FHierarchyElementViewModel> DraggedItem, EItemDropZone ItemDropZone)
 {
-	return false;
+	// Editability check
+	{
+		FResultWithUserFeedback Results = IsEditableByUser();
+		if(Results.bResult == false)
+		{
+			return Results;
+		}
+	}
+
+	// Self identity check
+	if(DraggedItem->GetData() == GetData())
+	{
+		return false;
+	}
+
+	// Can Contain checks
+	if(ItemDropZone == EItemDropZone::OntoItem)
+	{
+		FResultWithUserFeedback Results = CanContain(DraggedItem->GetData()->GetClass());
+		if(Results.bResult == false)
+		{
+			return Results;
+		}
+	}
+	else
+	{
+		if(Parent.IsValid())
+		{
+			FResultWithUserFeedback Results = Parent.Pin()->CanContain(DraggedItem->GetData()->GetClass());
+			if(Results.bResult == false)
+			{
+				return Results;
+			}
+		}
+	}
+
+	// Internal check
+	return CanDropOnInternal(DraggedItem, ItemDropZone);
+}
+
+void FHierarchyElementViewModel::OnDroppedOn(TSharedPtr<FHierarchyElementViewModel> DroppedItem, EItemDropZone ItemDropZone)
+{
+	if(OnDroppedOnOverride(DroppedItem, ItemDropZone))
+	{
+		return;
+	}
+	
+	FScopedTransaction Transaction(LOCTEXT("Transaction_MovedItem", "Moved an item in the hierarchy"));
+	HierarchyViewModel->GetHierarchyRoot()->Modify();
+	
+	bool bDropSucceeded = false;
+	TSharedPtr<FHierarchyElementViewModel> NewItemViewModel;
+	if(ItemDropZone == EItemDropZone::OntoItem)
+	{
+		if(DroppedItem->IsForHierarchy() == false)
+		{
+			NewItemViewModel = DuplicateToThis(DroppedItem);
+		}
+		else
+		{
+			NewItemViewModel = ReparentToThis(DroppedItem);
+		}
+
+		bDropSucceeded = true;
+	}
+	else if(ItemDropZone == EItemDropZone::AboveItem)
+	{
+		int32 IndexOfThis = Parent.Pin()->FindIndexOfDataChild(AsShared());
+
+		if(DroppedItem->IsForHierarchy() == false)
+		{
+			NewItemViewModel = Parent.Pin()->DuplicateToThis(DroppedItem, FMath::Max(IndexOfThis, 0));
+		}
+		else
+		{
+			NewItemViewModel = Parent.Pin()->ReparentToThis(DroppedItem, FMath::Max(IndexOfThis, 0));
+		}
+
+		bDropSucceeded = true;
+	}
+	else if(ItemDropZone == EItemDropZone::BelowItem)
+	{
+		int32 IndexOfThis = Parent.Pin()->FindIndexOfDataChild(AsShared());
+		
+		if(DroppedItem->IsForHierarchy() == false)
+		{
+			NewItemViewModel = Parent.Pin()->DuplicateToThis(DroppedItem, FMath::Min(IndexOfThis+1, Parent.Pin()->GetChildren().Num()));
+		}
+		else
+		{
+			NewItemViewModel = Parent.Pin()->ReparentToThis(DroppedItem, FMath::Min(IndexOfThis+1, Parent.Pin()->GetChildren().Num()));
+		}
+
+		bDropSucceeded = true;
+	}
+	
+	if(bDropSucceeded)
+	{
+		NewItemViewModel->PostAddFixup();
+		HierarchyViewModel->RefreshHierarchyView();
+		HierarchyViewModel->RefreshSourceView();
+	}
+	else
+	{
+		Transaction.Cancel();
+	}
+}
+
+FHierarchyElementViewModel::FResultWithUserFeedback FHierarchyElementViewModel::CanContain(TSubclassOf<UHierarchyElement> HierarchyElementType)
+{
+	// Editability check
+	{
+		FResultWithUserFeedback Results = IsEditableByUser();
+		if(Results.bResult == false)
+		{
+			return Results;
+		}
+	}
+
+	// Allow children check
+	{
+		FResultWithUserFeedback Results = CanHaveChildren();
+		if(Results.bResult == false)
+		{
+			return Results;
+		}
+	}
+
+	// Internal check
+	return CanContainInternal(HierarchyElementType);
+}
+
+FHierarchyElementViewModel::FResultWithUserFeedback FHierarchyElementViewModel::CanDropOnInternal(TSharedPtr<FHierarchyElementViewModel>, EItemDropZone ItemDropZone)
+{
+	return true;
 }
 
 void FHierarchyElementViewModel::PropagateOnChildSynced()
@@ -1139,7 +1299,7 @@ void FHierarchyElementViewModel::PropagateOnChildSynced()
 
 FReply FHierarchyElementViewModel::OnDragDetected(const FGeometry& Geometry, const FPointerEvent& PointerEvent,	bool bIsSource)
 {
-	FCanPerformActionResults CanDragResults = CanDrag();
+	FResultWithUserFeedback CanDragResults = CanDrag();
 	if(CanDragResults == true)
 	{
 		// if the drag is coming from source, we check if any of the hierarchy data already contains that element and we don't start a drag drop in that case
@@ -1176,9 +1336,9 @@ FReply FHierarchyElementViewModel::OnDragDetected(const FGeometry& Geometry, con
 	else
 	{
 		// if we can't drag and have a message, we show it as a slate notification
-		if(CanDragResults.CanPerformMessage.IsEmpty() == false)
+		if(CanDragResults.UserFeedback.IsSet())
 		{
-			FNotificationInfo CantDragInfo(CanDragResults.CanPerformMessage);
+			FNotificationInfo CantDragInfo(CanDragResults.UserFeedback.GetValue());
 			FSlateNotificationManager::Get().AddNotification(CantDragInfo);
 		}
 	}
@@ -1196,75 +1356,53 @@ void FHierarchyRootViewModel::Initialize()
 	GetOnSynced().BindSP(this, &FHierarchyRootViewModel::PropagateOnSynced);
 }
 
-FHierarchyElementViewModel::FCanPerformActionResults FHierarchyRootViewModel::CanDropOnInternal(TSharedPtr<FHierarchyElementViewModel> DraggedElement, EItemDropZone ItemDropZone)
+FHierarchyElementViewModel::FResultWithUserFeedback FHierarchyRootViewModel::CanContainInternal(TSubclassOf<UHierarchyElement> InHierarchyElement)
 {
-	FCanPerformActionResults Results(false);
+	return true;
+}
+
+FHierarchyElementViewModel::FResultWithUserFeedback FHierarchyRootViewModel::CanDropOnInternal(TSharedPtr<FHierarchyElementViewModel> DraggedElement, EItemDropZone ItemDropZone)
+{
+	FResultWithUserFeedback Results(false);
 	
 	// we only allow drops if some general conditions are fulfilled
 	if(DraggedElement->GetData() != GetData() &&
 		(!DraggedElement->HasParent(AsShared(), false) || ItemDropZone != EItemDropZone::OntoItem)  &&
 		!HasParent(DraggedElement, true))
 	{
-		Results.bCanPerform = 
+		Results.bResult = 
 			// items can be dropped onto the root directly if the section is set to "All"
-			(DraggedElement->GetData()->IsA<UHierarchyItem>() && HierarchyViewModel->GetActiveHierarchySectionData() == nullptr)
+			(DraggedElement->GetData()->IsA<UHierarchyItem>() && HierarchyViewModel->GetActiveHierarchySectionViewModel()->IsAllSection())
 				||
 			// categories can be dropped onto the root always
 			(DraggedElement->GetData()->IsA<UHierarchyCategory>());
 
-		if(Results.bCanPerform)
+		if(Results.bResult)
 		{
 			if(DraggedElement->IsForHierarchy() == false)
 			{
 				FText Message = LOCTEXT("CanDropSourceItemOnRootDragMessage", "Add {0} to the hierarchy root.");
 				Message = FText::FormatOrdered(Message, FText::FromString(DraggedElement->ToString()));
-				Results.CanPerformMessage = Message;
+				Results.UserFeedback = Message;
 			}
 			else
 			{
 				FText Message = LOCTEXT("CanDropHierarchyItemOnRootDragMessage", "Move {0} to the hierarchy root.");
 				Message = FText::FormatOrdered(Message, FText::FromString(DraggedElement->ToString()));
-				Results.CanPerformMessage = Message;
+				Results.UserFeedback = Message;
 			}			
 		}
 		else
 		{
 			FText Message = LOCTEXT("CantDropHierarchyItemOnRootDragMessage", "Can not add {0} here. Please add it to a category!");
 			Message = FText::FormatOrdered(Message, FText::FromString(DraggedElement->ToString()));
-			Results.CanPerformMessage = Message;
+			Results.UserFeedback = Message;
 		}
 	}
 	
 	return Results;
 }
 
-void FHierarchyRootViewModel::OnDroppedOnInternal(TSharedPtr<FHierarchyElementViewModel> DroppedItem, EItemDropZone ItemDropZone)
-{
-	FScopedTransaction Transaction(LOCTEXT("Transaction_OnDropOnRoot", "Dropped item on root"));
-	HierarchyViewModel->GetHierarchyRoot()->Modify();
-
-	if(DroppedItem->GetDataMutable()->IsA<UHierarchyItem>() || DroppedItem->GetDataMutable()->IsA<UHierarchyCategory>())
-	{
-		TSharedPtr<FHierarchyElementViewModel> NewViewModel;
-		// we duplicate the item if the dragged item is from source
-		if(DroppedItem->IsForHierarchy() == false)
-		{
-			NewViewModel = DuplicateToThis(DroppedItem);
-		}
-		else
-		{
-			NewViewModel = ReparentToThis(DroppedItem);
-		}
-
-		if(UHierarchyCategory* AsCategory = Cast<UHierarchyCategory>(NewViewModel->GetDataMutable()))
-		{
-			AsCategory->SetSection(HierarchyViewModel->GetActiveHierarchySectionData());
-		}
-
-		HierarchyViewModel->RefreshHierarchyView();
-	}
-}
-
 TSharedPtr<FHierarchySectionViewModel> FHierarchyRootViewModel::AddSection()
 {
 	FScopedTransaction ScopedTransaction(LOCTEXT("NewSectionAdded","Added Section"));
@@ -1363,6 +1501,21 @@ void FHierarchyRootViewModel::SyncViewModelsToDataInternal()
 		});
 }
 
+FHierarchySectionViewModel::FHierarchySectionViewModel(UHierarchySection* InItem, TSharedRef<FHierarchyRootViewModel> InParent, TWeakObjectPtr<UDataHierarchyViewModelBase> InHierarchyViewModel)
+	: FHierarchyElementViewModel(InItem, InParent, InHierarchyViewModel, InParent->IsForHierarchy())
+{
+}
+
+UHierarchySection* FHierarchySectionViewModel::GetSectionData() const
+{
+	if(IsAllSection())
+	{
+		return nullptr;
+	}
+
+	return GetDataMutable<UHierarchySection>();
+}
+
 FString FHierarchySectionViewModel::ToString() const
 {
 	return GetSectionNameAsText().ToString();
@@ -1390,139 +1543,68 @@ void FHierarchySectionViewModel::SetSectionNameAsText(const FText& Text)
 
 FText FHierarchySectionViewModel::GetSectionNameAsText() const
 {
-	if(UHierarchySection* Section = Cast<UHierarchySection>(Element))
-	{
-		return Section->GetSectionNameAsText();
-	}
-
-	return LOCTEXT("DefaultSectionName", "All");
+	return Cast<UHierarchySection>(Element)->GetSectionNameAsText();
 }
 
 FText FHierarchySectionViewModel::GetSectionTooltip() const
 {
-	if(UHierarchySection* Section = Cast<UHierarchySection>(Element))
-	{
-		return Section->GetTooltip();
-	}
-	
-	return FText::GetEmpty();
+	return Cast<UHierarchySection>(Element)->GetTooltip();
 }
 
-FHierarchyElementViewModel::FCanPerformActionResults FHierarchySectionViewModel::CanDragInternal()
+bool FHierarchySectionViewModel::AllowEditingInDetailsPanel() const
 {
-	// We only allow hierarchy sections to be dragged, excluding the All section that has no valid data
-	return IsForHierarchy() && GetData() != nullptr;
+	return IsForHierarchy() && !IsAllSection();
 }
 
-bool FHierarchySectionViewModel::CanRenameInternal()
+FHierarchyElementViewModel::FResultWithUserFeedback FHierarchySectionViewModel::IsEditableByUser()
 {
-	return IsForHierarchy() && GetData() != nullptr;
+	// We generally leave the All section for the hierarchy editable as we still need to allow drag & drop on it
+	return FResultWithUserFeedback(IsForHierarchy());
 }
 
-bool FHierarchySectionViewModel::CanDeleteInternal()
+FHierarchyElementViewModel::FResultWithUserFeedback FHierarchySectionViewModel::CanDragInternal()
 {
-	return IsForHierarchy() && GetData() != nullptr;
+	return IsForHierarchy() && !IsAllSection();
 }
 
-FHierarchyElementViewModel::FCanPerformActionResults FHierarchySectionViewModel::CanDropOnInternal(TSharedPtr<FHierarchyElementViewModel> DraggedElement, EItemDropZone ItemDropZone)
+bool FHierarchySectionViewModel::CanRenameInternal()
 {
-	if(bDropDisallowed)
-	{
-		return false;
-	}
-	
-	FCanPerformActionResults Results(false);
-	// we don't allow dropping onto source sections and we don't specify a message as the sections aren't going to light up as valid drop targets
-	if(IsForHierarchy() == false)
-	{
-		return false;
-	}
-	
-	if(const UHierarchyCategory* Category = Cast<UHierarchyCategory>(DraggedElement->GetData()))
-	{
-		if(ItemDropZone == EItemDropZone::OntoItem)
-		{
-			FText Message = LOCTEXT("DropCategoryOnSectionDragMessage", "Add {0} to section {1}");
-			Message = FText::FormatOrdered(Message, FText::FromString(DraggedElement->ToString()), FText::FromString(ToString()));
-			
-			Results.bCanPerform = GetData() != Category->GetSection();
-			Results.CanPerformMessage = Results.bCanPerform ? Message : FText::GetEmpty();
-		}
-	}
-	else if(UHierarchySection* DraggedSection = Cast<UHierarchySection>(DraggedElement->GetDataMutable()))
-	{
-		const bool bSameSection = GetData() == DraggedSection;
-
-		// If we drag a section onto a section, nothing happens
-		if(ItemDropZone == EItemDropZone::OntoItem)
-		{
-			Results.bCanPerform = false;
-			return Results;
-		}
-
-		// The 'All' section does not accept any drop actions.
-		if(GetData() == nullptr)
-		{
-			Results.bCanPerform = false;
-			return Results;
-		}
-		
-		int32 DraggedSectionIndex = GetHierarchyViewModel()->GetHierarchyRoot()->GetSectionIndex(DraggedSection);
-		int32 InsertionIndex = GetHierarchyViewModel()->GetHierarchyRoot()->GetSectionIndex(GetDataMutable<UHierarchySection>());
-		// we add 1 to the insertion index if it's below an item because we either want to insert at the current index to place the item above, or at current+1 for below
-		InsertionIndex += ItemDropZone == EItemDropZone::AboveItem ? -1 : 1;
-
-		Results.bCanPerform = !bSameSection && DraggedSectionIndex != InsertionIndex;
-
-		if(Results.bCanPerform)
-		{
-			if(ItemDropZone != EItemDropZone::OntoItem)
-			{
-				FText Message = LOCTEXT("MoveSectionLeftDragMessage", "Move section here");
-				Message = FText::FormatOrdered(Message, FText::FromString(DraggedElement->ToString()));
-				Results.CanPerformMessage = Message;
-			}
-		}
-	}
-	else if(UHierarchyItem* Item = Cast<UHierarchyItem>(DraggedElement->GetDataMutable()))
-	{
-		FText Message = LOCTEXT("CantDropItemOnSectionDragMessage", "Can't drop items onto sections. Please drag a category onto section {0}");
-		Message = FText::FormatOrdered(Message, FText::FromString(ToString()));
-		Results.bCanPerform = false;
-		Results.CanPerformMessage = Message;
-	}
+	return IsForHierarchy() && !IsAllSection();
+}
 
-	return Results;
+bool FHierarchySectionViewModel::CanDeleteInternal()
+{
+	return IsForHierarchy() && !IsAllSection();
 }
 
-void FHierarchySectionViewModel::OnDroppedOnInternal(TSharedPtr<FHierarchyElementViewModel> DroppedItem, EItemDropZone ItemDropZone)
+bool FHierarchySectionViewModel::OnDroppedOnOverride(TSharedPtr<FHierarchyElementViewModel> DroppedItem, EItemDropZone ItemDropZone)
 {
 	if(DroppedItem->GetData()->IsA<UHierarchySection>())
 	{
 		FScopedTransaction Transaction(LOCTEXT("Transaction_OnSectionMoved", "Moved section"));
 		HierarchyViewModel->GetHierarchyRoot()->Modify();
-
+	
 		UHierarchySection* DraggedSectionData = DroppedItem->GetDataMutable<UHierarchySection>();
-
+	
 		int32 IndexOfThis = HierarchyViewModel->GetHierarchyRoot()->GetSectionData().Find(GetDataMutable<UHierarchySection>());
 		int32 DraggedSectionIndex = HierarchyViewModel->GetHierarchyRoot()->GetSectionData().Find(DraggedSectionData);
-
+	
 		TArray<TObjectPtr<UHierarchySection>>& SectionData = HierarchyViewModel->GetHierarchyRoot()->GetSectionDataMutable();
 		int32 Count = SectionData.Num();
-
+	
 		bool bDropSucceeded = false;
 		// above constitutes to the left here
 		if(ItemDropZone == EItemDropZone::AboveItem)
 		{
 			SectionData.RemoveAt(DraggedSectionIndex);
 			SectionData.Insert(DraggedSectionData, FMath::Max(IndexOfThis, 0));
-
+	
 			bDropSucceeded = true;
 		}
 		else if(ItemDropZone == EItemDropZone::BelowItem)
 		{
 			SectionData.RemoveAt(DraggedSectionIndex);
-
+	
 			if(IndexOfThis + 1 > SectionData.Num())
 			{
 				SectionData.Add(DraggedSectionData);
@@ -1531,11 +1613,11 @@ void FHierarchySectionViewModel::OnDroppedOnInternal(TSharedPtr<FHierarchyElemen
 			{
 				SectionData.Insert(DraggedSectionData, FMath::Min(IndexOfThis+1, Count));
 			}
-
+	
 			bDropSucceeded = true;
-
+	
 		}
-
+	
 		if(bDropSucceeded)
 		{
 			HierarchyViewModel->ForceFullRefresh();
@@ -1547,8 +1629,8 @@ void FHierarchySectionViewModel::OnDroppedOnInternal(TSharedPtr<FHierarchyElemen
 		FScopedTransaction Transaction(LOCTEXT("Transaction_OnSectionDrop", "Moved category to section"));
 		HierarchyViewModel->GetHierarchyRoot()->Modify();
 		
-		HierarchyCategory->SetSection(GetDataMutable<UHierarchySection>());
-
+		HierarchyCategory->SetSection(IsAllSection() ? nullptr : GetDataMutable<UHierarchySection>());
+	
 		// we null out any sections for all contained categories
 		TArray<UHierarchyCategory*> AllChildCategories;
 		HierarchyCategory->GetChildrenOfType<UHierarchyCategory>(AllChildCategories, true);
@@ -1556,7 +1638,7 @@ void FHierarchySectionViewModel::OnDroppedOnInternal(TSharedPtr<FHierarchyElemen
 		{
 			ChildCategory->SetSection(nullptr);
 		}
-
+	
 		// we only need to reparent if the parent isn't already the root. This stops unnecessary reordering
 		if(DroppedItem->GetParent() != HierarchyViewModel->GetHierarchyRootViewModel())
 		{
@@ -1566,11 +1648,92 @@ void FHierarchySectionViewModel::OnDroppedOnInternal(TSharedPtr<FHierarchyElemen
 		HierarchyViewModel->RefreshHierarchyView();
 		HierarchyViewModel->OnHierarchyChanged().Broadcast();
 	}
+
+	return true;
+}
+
+FHierarchyElementViewModel::FResultWithUserFeedback FHierarchySectionViewModel::CanContainInternal(TSubclassOf<UHierarchyElement> InHierarchyElementType)
+{
+	// While a section doesn't actually 'contain' a category as a child, semantically it gets treated as such in the drag & drop logic
+	if(InHierarchyElementType->IsChildOf<UHierarchyCategory>())
+	{
+		return true;
+	}
+
+	FResultWithUserFeedback Results(false);
+	Results.UserFeedback = LOCTEXT("SectionsCanOnlyContainCategoriesFeedback", "Sections can only contain categories.");
+	return Results;
+}
+
+FHierarchyElementViewModel::FResultWithUserFeedback FHierarchySectionViewModel::CanDropOnInternal(TSharedPtr<FHierarchyElementViewModel> DraggedElement, EItemDropZone ItemDropZone)
+{
+	FResultWithUserFeedback Results(false);
+	// we don't allow dropping onto source sections and we don't specify a message as the sections aren't going to light up as valid drop targets
+	if(IsForHierarchy() == false)
+	{
+		return false;
+	}
+	
+	if(const UHierarchyCategory* DraggedCategory = Cast<UHierarchyCategory>(DraggedElement->GetData()))
+	{
+		if(ItemDropZone == EItemDropZone::OntoItem)
+		{
+			FText Message = LOCTEXT("DropCategoryOnSectionDragMessage", "Add {0} to section {1}");
+			Message = FText::FormatOrdered(Message, FText::FromString(DraggedElement->ToString()), FText::FromString(ToString()));
+			
+			Results.bResult = GetData() != DraggedCategory->GetSection();
+			Results.UserFeedback = Results.bResult ? Message : FText::GetEmpty();
+		}
+	}
+	else if(UHierarchySection* DraggedSection = Cast<UHierarchySection>(DraggedElement->GetDataMutable()))
+	{
+		TSharedPtr<FHierarchySectionViewModel> DraggedSectionViewModel = StaticCastSharedPtr<FHierarchySectionViewModel>(DraggedElement);
+		const bool bSameSection = GetSectionData() == DraggedSectionViewModel->GetSectionData();
+
+		// If we drag a section onto a section, nothing happens
+		if(ItemDropZone == EItemDropZone::OntoItem)
+		{
+			Results.bResult = false;
+			return Results;
+		}
+
+		// The 'All' section does not accept any section drop actions.
+		if(IsAllSection())
+		{
+			Results.bResult = false;
+			return Results;
+		}
+		
+		int32 DraggedSectionIndex = GetHierarchyViewModel()->GetHierarchyRoot()->GetSectionIndex(DraggedSection);
+		int32 ThisSectionIndex = GetHierarchyViewModel()->GetHierarchyRoot()->GetSectionIndex(GetDataMutable<UHierarchySection>());
+
+		// We add +1 if below the item, or -1 if above the item, to determine if we'd land on the same index as we currently are.
+		// For example, if we have [Section1] [Section2] [Section3], and we drag [Section2] onto the left part of [Section3] or right part of [Section1], nothing should happen
+		ThisSectionIndex += ItemDropZone == EItemDropZone::AboveItem ? -1 : 1;
+		
+		Results.bResult = !bSameSection && DraggedSectionIndex != ThisSectionIndex;
+
+		if(Results.bResult)
+		{
+			if(ItemDropZone != EItemDropZone::OntoItem)
+			{
+				FText Message = LOCTEXT("MoveSectionLeftDragMessage", "Move section {0} here");
+				Message = FText::FormatOrdered(Message, FText::FromString(DraggedElement->ToString()));
+				Results.UserFeedback = Message;
+			}
+		}
+	}
+	else if(UHierarchyItem* Item = Cast<UHierarchyItem>(DraggedElement->GetDataMutable()))
+	{
+		return false;
+	}
+
+	return Results;
 }
 
 void FHierarchySectionViewModel::FinalizeInternal()
 {
-	if(HierarchyViewModel->GetActiveHierarchySection() == AsShared())
+	if(HierarchyViewModel->GetActiveHierarchySectionViewModel() == AsShared())
 	{
 		HierarchyViewModel->SetActiveHierarchySection(HierarchyViewModel->GetDefaultHierarchySectionViewModel());
 	}
@@ -1588,7 +1751,12 @@ void FHierarchySectionViewModel::FinalizeInternal()
 	}
 }
 
-::FHierarchyElementViewModel::FCanPerformActionResults FHierarchyItemViewModel::CanDropOnInternal(TSharedPtr<FHierarchyElementViewModel> DraggedElement, EItemDropZone ItemDropZone)
+FHierarchyElementViewModel::FResultWithUserFeedback FHierarchyItemViewModel::CanContainInternal(TSubclassOf<UHierarchyElement> InHierarchyElement)
+{
+	return false;
+}
+
+::FHierarchyElementViewModel::FResultWithUserFeedback FHierarchyItemViewModel::CanDropOnInternal(TSharedPtr<FHierarchyElementViewModel> DraggedElement, EItemDropZone ItemDropZone)
 {
 	bool bAllowDrop = false;
 	
@@ -1607,130 +1775,49 @@ void FHierarchySectionViewModel::FinalizeInternal()
 	return bAllowDrop;
 }
 
-void FHierarchyItemViewModel::OnDroppedOnInternal(TSharedPtr<FHierarchyElementViewModel> DroppedItem, EItemDropZone ItemDropZone)
+const FTableRowStyle* FHierarchyItemViewModel::GetRowStyle() const
 {
-	FScopedTransaction Transaction(LOCTEXT("Transaction_MovedItem", "Moved an item in the hierarchy"));
-	HierarchyViewModel->GetHierarchyRoot()->Modify();
-	
-	bool bDropSucceeded = false;
-	if(ItemDropZone == EItemDropZone::AboveItem)
-	{
-		int32 IndexOfThis = Parent.Pin()->FindIndexOfDataChild(AsShared());
-
-		if(DroppedItem->IsForHierarchy() == false)
-		{
-			Parent.Pin()->DuplicateToThis(DroppedItem, FMath::Max(IndexOfThis, 0));
-		}
-		else
-		{
-			Parent.Pin()->ReparentToThis(DroppedItem, FMath::Max(IndexOfThis, 0));
-		}
-
-		bDropSucceeded = true;
-	}
-	else if(ItemDropZone == EItemDropZone::BelowItem)
-	{
-		int32 IndexOfThis = Parent.Pin()->FindIndexOfDataChild(AsShared());
-		
-		if(DroppedItem->IsForHierarchy() == false)
-		{
-			Parent.Pin()->DuplicateToThis(DroppedItem, FMath::Min(IndexOfThis+1, Parent.Pin()->GetChildren().Num()));
-		}
-		else
-		{
-			Parent.Pin()->ReparentToThis(DroppedItem, FMath::Min(IndexOfThis+1, Parent.Pin()->GetChildren().Num()));
-		}
-
-		bDropSucceeded = true;
-	}
-
-	if(bDropSucceeded)
-	{
-		HierarchyViewModel->RefreshHierarchyView();
-		HierarchyViewModel->RefreshSourceView();
-	}
-	else
-	{
-		Transaction.Cancel();
-	}
+	return &FAppStyle::Get().GetWidgetStyle<FTableRowStyle>("TableView.Row");
 }
 
-void FHierarchyCategoryViewModel::OnDroppedOnInternal(TSharedPtr<FHierarchyElementViewModel> DroppedItem, EItemDropZone ItemDropZone)
+void FHierarchyCategoryViewModel::PostAddFixup()
 {
-	FScopedTransaction Transaction(LOCTEXT("Transaction_OnCategoryDrop", "Dropped item on/above/below category"));
-	HierarchyViewModel->GetHierarchyRoot()->Modify();
-	
-	if(UHierarchyCategory* Category = DroppedItem->GetDataMutable<UHierarchyCategory>())
+	// If the new category is under the root, we set the section to the currently active one
+	if(Parent.Pin()->GetData()->IsA<UHierarchyRoot>())
 	{
-		// if we are dragging a category above/below another category and the new parent is going to be the root, we update its section to the active section
-		if(ItemDropZone != EItemDropZone::OntoItem)
-		{
-			if(Parent.IsValid() && Parent == HierarchyViewModel->GetHierarchyRootViewModel())
-			{
-				Category->SetSection(HierarchyViewModel->GetActiveHierarchySectionData());
-
-				// we null out any sections for all contained categories
-				TArray<UHierarchyCategory*> AllChildCategories;
-				Category->GetChildrenOfType<UHierarchyCategory>(AllChildCategories, true);
-				for(UHierarchyCategory* ChildCategory : AllChildCategories)
-				{
-					ChildCategory->SetSection(nullptr);
-				}
-			}				
-		}
-		// if we are dragging a category onto another category, we null out its section instead
-		else
-		{
-			Category->SetSection(nullptr);
+		UHierarchyCategory* Category = GetDataMutable<UHierarchyCategory>();
+		Category->SetSection(HierarchyViewModel->GetActiveHierarchySectionViewModel()->IsAllSection() ? nullptr : HierarchyViewModel->GetActiveHierarchySectionData());
 
-			// we null out any sections for all contained categories
-			TArray<UHierarchyCategory*> AllChildCategories;
-			Category->GetChildrenOfType<UHierarchyCategory>(AllChildCategories, true);
-			for(UHierarchyCategory* ChildCategory : AllChildCategories)
-			{
-				ChildCategory->SetSection(nullptr);
-			}
-		}			
-	}
-
-	// the actual moving of the item happens here
-	if(ItemDropZone == EItemDropZone::OntoItem)
-	{
-		if(DroppedItem->IsForHierarchy() == false)
-		{
-			DuplicateToThis(DroppedItem);
-		}
-		else
-		{
-			ReparentToThis(DroppedItem);
-		}
-	}
-	else if(ItemDropZone == EItemDropZone::AboveItem)
-	{
-		int32 IndexOfThis = Parent.Pin()->FindIndexOfDataChild(AsShared());
-		if(DroppedItem->IsForHierarchy() == false)
+		// we null out any sections for all contained categories
+		TArray<UHierarchyCategory*> AllChildCategories;
+		Category->GetChildrenOfType<UHierarchyCategory>(AllChildCategories, true);
+		for(UHierarchyCategory* ChildCategory : AllChildCategories)
 		{
-			Parent.Pin()->DuplicateToThis(DroppedItem, FMath::Max(IndexOfThis, 0));
-		}
-		else
-		{				
-			Parent.Pin()->ReparentToThis(DroppedItem, FMath::Max(IndexOfThis, 0));
+			ChildCategory->SetSection(nullptr);
 		}
 	}
-	else if(ItemDropZone == EItemDropZone::BelowItem)
+	// If it's a subcategory, we empty out the section
+	else if(Parent.Pin()->GetData()->IsA<UHierarchyCategory>())
 	{
-		int32 IndexOfThis = Parent.Pin()->FindIndexOfDataChild(AsShared());
-		if(DroppedItem->IsForHierarchy() == false)
-		{
-			Parent.Pin()->DuplicateToThis(DroppedItem, FMath::Min(IndexOfThis+1, Parent.Pin()->GetChildren().Num()));
-		}
-		else
+		UHierarchyCategory* Category = GetDataMutable<UHierarchyCategory>();
+
+		Category->SetSection(nullptr);
+
+		// we null out any sections for all contained categories
+		TArray<UHierarchyCategory*> AllChildCategories;
+		Category->GetChildrenOfType<UHierarchyCategory>(AllChildCategories, true);
+		for(UHierarchyCategory* ChildCategory : AllChildCategories)
 		{
-			Parent.Pin()->ReparentToThis(DroppedItem, FMath::Min(IndexOfThis+1, Parent.Pin()->GetChildren().Num()));
+			ChildCategory->SetSection(nullptr);
 		}
 	}
 }
 
+const FTableRowStyle* FHierarchyCategoryViewModel::GetRowStyle() const
+{
+	return &FDataHierarchyEditorStyle::Get().GetWidgetStyle<FTableRowStyle>("HierarchyEditor.Row.Category");
+}
+
 void UDataHierarchyViewModelBase::AddCategory(TSharedPtr<FHierarchyElementViewModel> CategoryParent) const
 {
 	// If no category parent was specified, we add it to the root
@@ -1771,7 +1858,7 @@ void UDataHierarchyViewModelBase::AddCategory(TSharedPtr<FHierarchyElementViewMo
 		Category->SetCategoryName(UE::DataHierarchyEditor::GetUniqueName(FName("New Category"), CategoryNames));
 		
 		// we only set the section property if the current section isn't set to "All"
-		Category->SetSection(GetActiveHierarchySectionData());
+		Category->SetSection(GetActiveHierarchySectionViewModel()->IsAllSection() ? nullptr : GetActiveHierarchySectionData());
 		
 		RefreshHierarchyView();
 
@@ -1813,11 +1900,49 @@ void UDataHierarchyViewModelBase::GenerateDynamicContextMenuInternal(UToolMenu*
 		return;
 	}
 
-	UDataHierarchyViewModelBase* HierarchyViewModel = HierarchyMenuContext->HierarchyViewModel.Get();
+	// Disclaimer: We aren't simply using the commands with the command list, because the regular commandlist logic does not allow us to pass in tool menu contexts.
+	// CommandLists work with shortcuts, meaning there might not be a menu to work with.
+	// This means for something like a 'Delete' action, we have to figure out 'what' to delete without context, which is difficult since we have 2 tree views & selected sections.
+	// When we DO have a menu, we want to use the better ToolMenu logic that tells us what elements we are operating on.
+	
+	// Find
+	{
+		FToolUIAction Action;
+		Action.ExecuteAction = FToolMenuExecuteAction::CreateUObject(this, &UDataHierarchyViewModelBase::ToolMenuNavigateTo);
+		Action.CanExecuteAction = FToolMenuCanExecuteAction::CreateUObject(this, &UDataHierarchyViewModelBase::ToolMenuCanNavigateTo);
+		Action.IsActionVisibleDelegate = FToolMenuIsActionButtonVisible::CreateUObject(this, &UDataHierarchyViewModelBase::ToolMenuCanNavigateTo);
+
+		TSharedPtr<FUICommandInfo> Command = FDataHierarchyEditorCommands::Get().FindInHierarchy;
+		FToolMenuEntry MenuEntry = FToolMenuEntry::InitMenuEntry(Command->GetCommandName(), Command->GetLabel(), Command->GetDescription(), Command->GetIcon(), Action);
+		MenuEntry.InputBindingLabel = Command->GetInputText();
+		DynamicToolMenu->AddMenuEntry("Dynamic", MenuEntry);
+	}
+	
+	// Rename
+	{
+		FToolUIAction Action;
+		Action.ExecuteAction = FToolMenuExecuteAction::CreateUObject(this, &UDataHierarchyViewModelBase::ToolMenuRequestRename);
+		Action.CanExecuteAction = FToolMenuCanExecuteAction::CreateUObject(this, &UDataHierarchyViewModelBase::ToolMenuCanRequestRename);
+		Action.IsActionVisibleDelegate = FToolMenuIsActionButtonVisible::CreateUObject(this, &UDataHierarchyViewModelBase::ToolMenuCanRequestRename);
+
+		TSharedPtr<FUICommandInfo> Command = FGenericCommands::Get().Rename;
+		FToolMenuEntry MenuEntry = FToolMenuEntry::InitMenuEntry(Command->GetCommandName(), Command->GetLabel(), Command->GetDescription(), Command->GetIcon(), Action);
+		MenuEntry.InputBindingLabel = Command->GetInputText();
+		DynamicToolMenu->AddMenuEntry("Dynamic", MenuEntry);
+	}
 
-	DynamicToolMenu->AddMenuEntry("Dynamic", FToolMenuEntry::InitMenuEntryWithCommandList(FDataHierarchyEditorCommands::Get().FindInHierarchy, HierarchyViewModel->GetCommands(), TAttribute<FText>(), TAttribute<FText>(), FSlateIcon(FAppStyle::GetAppStyleSetName(), "Icons.Find")));
-	DynamicToolMenu->AddMenuEntry("Dynamic", FToolMenuEntry::InitMenuEntryWithCommandList(FGenericCommands::Get().Rename, HierarchyViewModel->GetCommands()));
-	DynamicToolMenu->AddMenuEntry("Dynamic", FToolMenuEntry::InitMenuEntryWithCommandList(FGenericCommands::Get().Delete, HierarchyViewModel->GetCommands()));
+	// Delete
+	{
+		FToolUIAction Action;
+		Action.ExecuteAction = FToolMenuExecuteAction::CreateUObject(this, &UDataHierarchyViewModelBase::ToolMenuDelete);
+		Action.CanExecuteAction = FToolMenuCanExecuteAction::CreateUObject(this, &UDataHierarchyViewModelBase::ToolMenuCanDelete);
+		Action.IsActionVisibleDelegate = FToolMenuIsActionButtonVisible::CreateUObject(this, &UDataHierarchyViewModelBase::ToolMenuCanDelete);
+		
+		TSharedPtr<FUICommandInfo> Command = FGenericCommands::Get().Delete;
+		FToolMenuEntry MenuEntry = FToolMenuEntry::InitMenuEntry(Command->GetCommandName(), Command->GetLabel(), Command->GetDescription(), Command->GetIcon(), Action);
+		MenuEntry.InputBindingLabel = Command->GetInputText();
+		DynamicToolMenu->AddMenuEntry("Dynamic", MenuEntry);
+	}
 }
 
 UHierarchyElement* UDataHierarchyViewModelBase::AddElementUnderRoot(TSubclassOf<UHierarchyElement> NewChildClass, FHierarchyElementIdentity ChildIdentity)
diff --git a/Engine/Source/Editor/DataHierarchyEditor/Private/Widgets/SDataHierarchyEditor.cpp b/Engine/Source/Editor/DataHierarchyEditor/Private/Widgets/SDataHierarchyEditor.cpp
--- a/Engine/Source/Editor/DataHierarchyEditor/Private/Widgets/SDataHierarchyEditor.cpp
+++ b/Engine/Source/Editor/DataHierarchyEditor/Private/Widgets/SDataHierarchyEditor.cpp
@@ -3,6 +3,7 @@
 #include "Widgets/SDataHierarchyEditor.h"
 
 #include "DataHierarchyEditorCommands.h"
+#include "DataHierarchyEditorMisc.h"
 #include "UObject/Package.h"
 #include "Framework/Commands/GenericCommands.h"
 #include "Styling/StyleColors.h"
@@ -22,7 +23,10 @@
 #include "Framework/Application/SlateApplication.h"
 #include "ToolMenus.h"
 #include "DataHierarchyEditorStyle.h"
+#include "SPositiveActionButton.h"
+#include "Styling/SlateIconFinder.h"
 #include "Widgets/Input/SSearchBox.h"
+#include "Widgets/Layout/SSpacer.h"
 
 #define LOCTEXT_NAMESPACE "DataHierarchyEditor"
 
@@ -40,53 +44,47 @@ TSharedRef<SWidget> SummonContextMenu(TArray<TSharedPtr<FHierarchyElementViewMod
 	return MenuWidget;
 }
 
-void SHierarchyCategory::Construct(const FArguments& InArgs,	TSharedPtr<FHierarchyCategoryViewModel> InCategoryViewModel)
+void SHierarchyElement::Construct(const FArguments& InArgs, TSharedPtr<struct FHierarchyElementViewModel> InViewModel)
 {
-	CategoryViewModelWeak = InCategoryViewModel;
+	ElementViewModelWeak = InViewModel;
 	
-	InCategoryViewModel->GetOnRequestRename().BindSP(this, &SHierarchyCategory::EnterEditingMode);
-
-	const UHierarchyCategory* CategoryData = InCategoryViewModel->GetData<UHierarchyCategory>();
-	SetToolTipText(TAttribute<FText>::CreateLambda([CategoryData]
-	{
-		return CategoryData->GetTooltip();
-	}));
+	InViewModel->GetOnRequestRename().BindSP(this, &SHierarchyElement::EnterEditingMode);
 	
 	ChildSlot
 	[
 		SAssignNew(InlineEditableTextBlock, SInlineEditableTextBlock)
-		.Style(&FDataHierarchyEditorStyle::Get().GetWidgetStyle<FInlineEditableTextBlockStyle>("HierarchyEditor.CategoryTextStyle"))
-		.Text(this, &SHierarchyCategory::GetCategoryText)
-		.OnTextCommitted(this, &SHierarchyCategory::OnRenameCategory)
-		.OnVerifyTextChanged(this, &SHierarchyCategory::OnVerifyCategoryRename)
+		.Style(InArgs._Style)
+		.Text(this, &SHierarchyElement::GetElementText)
+		.OnTextCommitted(this, &SHierarchyElement::OnRenameElement)
+		.OnVerifyTextChanged(this, &SHierarchyElement::OnVerifyRename)
 		.IsSelected(InArgs._IsSelected)
 	];
 }
 
-void SHierarchyCategory::EnterEditingMode() const
+void SHierarchyElement::EnterEditingMode() const
 {
-	if(CategoryViewModelWeak.Pin()->CanRename())
+	if(ElementViewModelWeak.Pin()->CanRename())
 	{
 		InlineEditableTextBlock->EnterEditingMode();
 	}
 }
 
-bool SHierarchyCategory::OnVerifyCategoryRename(const FText& NewName, FText& OutTooltip) const
+bool SHierarchyElement::OnVerifyRename(const FText& NewName, FText& OutTooltip) const
 {
-	TArray<TSharedPtr<FHierarchyCategoryViewModel>> SiblingCategoryViewModels;
-	CategoryViewModelWeak.Pin()->GetParent().Pin()->GetChildrenViewModelsForType<UHierarchyCategory, FHierarchyCategoryViewModel>(SiblingCategoryViewModels);
+	TArray<TSharedPtr<FHierarchyItemViewModel>> SiblingItemViewModels;
+	ElementViewModelWeak.Pin()->GetParent().Pin()->GetChildrenViewModelsForType<UHierarchyElement, FHierarchyItemViewModel>(SiblingItemViewModels);
 
-	if(GetCategoryText().ToString() != NewName.ToString())
+	if(GetElementText().ToString() != NewName.ToString())
 	{
-		TSet<FString> CategoryNames;
-		for(const auto& SiblingCategoryViewModel : SiblingCategoryViewModels)
+		TSet<FString> ItemNames;
+		for(const auto& SiblingItemViewModel : SiblingItemViewModels)
 		{
-			CategoryNames.Add(Cast<UHierarchyCategory>(SiblingCategoryViewModel->GetDataMutable())->GetCategoryName().ToString());
+			ItemNames.Add(Cast<UHierarchyElement>(SiblingItemViewModel->GetDataMutable())->ToString());
 		}
 
-		if(CategoryNames.Contains(NewName.ToString()))
+		if(ItemNames.Contains(NewName.ToString()))
 		{
-			OutTooltip = LOCTEXT("HierarchyCategoryCantRename_DuplicateOnLayer", "Another category of the same name already exists!");
+			OutTooltip = LOCTEXT("HierarchyItemCantRename_DuplicateOnLayer", "Another item of the same name already exists!");
 			return false;
 		}
 	}
@@ -94,19 +92,19 @@ bool SHierarchyCategory::OnVerifyCategoryRename(const FText& NewName, FText& Out
 	return true;
 }
 
-FText SHierarchyCategory::GetCategoryText() const
+FText SHierarchyElement::GetElementText() const
 {
-	return FText::FromString(CategoryViewModelWeak.Pin()->ToString());
+	return FText::FromString(ElementViewModelWeak.Pin()->ToString());
 }
 
-void SHierarchyCategory::OnRenameCategory(const FText& NewText, ETextCommit::Type) const
+void SHierarchyElement::OnRenameElement(const FText& NewText, ETextCommit::Type) const
 {
-	if(CategoryViewModelWeak.IsValid() && CategoryViewModelWeak.Pin()->GetCategoryName().EqualTo(NewText) == false)
+	if(ElementViewModelWeak.IsValid() && ElementViewModelWeak.Pin()->ToStringAsText().EqualTo(NewText) == false)
 	{
-		FScopedTransaction Transaction(LOCTEXT("Transaction_Rename_Category", "Renamed hierarchy category"));
-		CategoryViewModelWeak.Pin()->GetHierarchyViewModel()->GetHierarchyRoot()->Modify();
+		FScopedTransaction Transaction(LOCTEXT("Transaction_Rename_Item", "Renamed hierarchy item"));
+		ElementViewModelWeak.Pin()->GetHierarchyViewModel()->GetHierarchyRoot()->Modify();
 		
-		CategoryViewModelWeak.Pin()->Rename(FName(NewText.ToString()));
+		ElementViewModelWeak.Pin()->Rename(FName(NewText.ToString()));
 	}
 }
 
@@ -223,7 +221,7 @@ bool SHierarchySection::OnCanAcceptDrop(TSharedPtr<FDragDropOperation> DragDropO
 		TSharedPtr<FHierarchyDragDropOp> HierarchyDragDropOp = StaticCastSharedPtr<FHierarchyDragDropOp>(DragDropOperation);
 		TSharedPtr<FHierarchyElementViewModel> DraggedItem = HierarchyDragDropOp->GetDraggedElement().Pin();
 		
-		return SectionViewModelWeak.Pin()->CanDropOn(DraggedItem, InItemDropZone).bCanPerform;
+		return SectionViewModelWeak.Pin()->CanDropOn(DraggedItem, InItemDropZone).bResult;
 	}
 
 	return false;
@@ -308,7 +306,7 @@ void SHierarchySection::OnDragEnter(const FGeometry& MyGeometry, const FDragDrop
 
 FReply SHierarchySection::OnDragDetected(const FGeometry& MyGeometry, const FPointerEvent& MouseEvent)
 {
-	if(MouseEvent.IsMouseButtonDown(EKeys::LeftMouseButton) && SectionViewModelWeak.Pin()->CanDrag().bCanPerform)
+	if(MouseEvent.IsMouseButtonDown(EKeys::LeftMouseButton) && SectionViewModelWeak.Pin()->CanDrag().bResult)
 	{
 		TSharedRef<FSectionDragDropOp> SectionDragDropOp = MakeShared<FSectionDragDropOp>(SectionViewModelWeak.Pin());
 		SectionDragDropOp->Construct();
@@ -337,10 +335,10 @@ FReply SHierarchySection::OnDragOver(const FGeometry& MyGeometry, const FDragDro
 	if(TSharedPtr<FHierarchyDragDropOp> HierarchyDragDropOp = DragDropEvent.GetOperationAs<FHierarchyDragDropOp>())
 	{
 		TSharedPtr<FHierarchyElementViewModel> DraggedItem = HierarchyDragDropOp->GetDraggedElement().Pin();
-		FHierarchyElementViewModel::FCanPerformActionResults Results = SectionViewModelWeak.Pin()->CanDropOn(DraggedItem, ItemHoverZone);
-		HierarchyDragDropOp->SetDescription(Results.CanPerformMessage);
+		FHierarchyElementViewModel::FResultWithUserFeedback Results = SectionViewModelWeak.Pin()->CanDropOn(DraggedItem, ItemHoverZone);
+		HierarchyDragDropOp->SetDescription(Results.UserFeedback.Get(FText::GetEmpty()));
 
-		if(Results.bCanPerform)
+		if(Results.bResult)
 		{
 			CurrentItemDropZone = ItemHoverZone;
 		}
@@ -462,10 +460,10 @@ const FSlateBrush* SHierarchySection::GetDropIndicatorBrush(EItemDropZone InItem
 {
 	switch (InItemDropZone)
 	{
-		case EItemDropZone::AboveItem: return FDataHierarchyEditorStyle::Get().GetBrush("HierarchyEditor.Drop.Section.Above");
+		case EItemDropZone::AboveItem: return FDataHierarchyEditorStyle::Get().GetBrush("HierarchyEditor.Drop.Section.Left");
 		default:
 		case EItemDropZone::OntoItem: return FDataHierarchyEditorStyle::Get().GetBrush("HierarchyEditor.Drop.Section.Onto");
-		case EItemDropZone::BelowItem: return FDataHierarchyEditorStyle::Get().GetBrush("HierarchyEditor.Drop.Section.Below");
+		case EItemDropZone::BelowItem: return FDataHierarchyEditorStyle::Get().GetBrush("HierarchyEditor.Drop.Section.Right");
 	};
 }
 
@@ -498,14 +496,21 @@ void SDataHierarchyEditor::Construct(const FArguments& InArgs, TObjectPtr<UDataH
 	{
 		HierarchyViewModel->Initialize();
 	}
-	
-	SourceRoot.Reset(NewObject<UHierarchyRoot>(HierarchyViewModel->GetOuterForSourceRoot(), MakeUniqueObjectName(HierarchyViewModel->GetOuterForSourceRoot(), UHierarchyRoot::StaticClass()), RF_Transient));
-	TSharedPtr<FHierarchyElementViewModel> ViewModel = HierarchyViewModel->CreateViewModelForElement(SourceRoot.Get(), nullptr);
-	SourceRootViewModel = StaticCastSharedPtr<FHierarchyRootViewModel>(ViewModel);	
-	SourceRootViewModel->Initialize();
-	SourceRootViewModel->AddChildFilter(FHierarchyElementViewModel::FOnFilterChild::CreateSP(this, &SDataHierarchyEditor::FilterForSourceSection));
-	SourceRootViewModel->OnSyncPropagated().BindSP(this, &SDataHierarchyEditor::RequestRefreshSourceViewNextFrame, false);
-	SourceRootViewModel->OnSectionsChanged().BindSP(this, &SDataHierarchyEditor::RefreshSectionsView);
+
+	if(HierarchyViewModel->SupportsSourcePanel())
+	{
+		SourceRoot.Reset(NewObject<UHierarchyRoot>(HierarchyViewModel->GetOuterForSourceRoot(), MakeUniqueObjectName(HierarchyViewModel->GetOuterForSourceRoot(), UHierarchyRoot::StaticClass()), RF_Transient));
+		TSharedPtr<FHierarchyElementViewModel> ViewModel = HierarchyViewModel->CreateViewModelForElement(SourceRoot.Get(), nullptr);
+		SourceRootViewModel = StaticCastSharedPtr<FHierarchyRootViewModel>(ViewModel);	
+		SourceRootViewModel->Initialize();
+		SourceRootViewModel->AddChildFilter(FHierarchyElementViewModel::FOnFilterChild::CreateSP(this, &SDataHierarchyEditor::FilterForSourceSection));
+		SourceRootViewModel->OnSyncPropagated().BindSP(this, &SDataHierarchyEditor::RequestRefreshSourceViewNextFrame, false);
+		SourceRootViewModel->OnSectionsChanged().BindSP(this, &SDataHierarchyEditor::RefreshSectionsView);
+
+		AllSourceSection.Reset(NewObject<UHierarchySection>(HierarchyViewModel.Get(), UDataHierarchyViewModelBase::AllSectionSourceObjectName));
+		DefaultSourceSectionViewModel = MakeShared<FHierarchySectionViewModel>(AllSourceSection.Get(), SourceRootViewModel.ToSharedRef(), HierarchyViewModel);
+		DefaultSourceSectionViewModel->SetSectionName("All");
+	}
 
 	HierarchyViewModel->OnInitialized().BindSP(this, &SDataHierarchyEditor::Reinitialize);
 	HierarchyViewModel->OnNavigateToElementIdentityInHierarchyRequested().BindSP(this, &SDataHierarchyEditor::NavigateToHierarchyElement);
@@ -522,8 +527,6 @@ void SDataHierarchyEditor::Construct(const FArguments& InArgs, TObjectPtr<UDataH
 	
 	OnGenerateRowContentWidget = InArgs._OnGenerateRowContentWidget;
 	OnGenerateCustomDetailsPanelNameWidget = InArgs._OnGenerateCustomDetailsPanelNameWidget;
-	CategoryRowStyle = InArgs._CategoryRowStyle;
-	ItemRowStyle = InArgs._ItemRowStyle;
 	
 	if(!ensureMsgf(OnGenerateRowContentWidget.IsBound(), TEXT("Please add a function binding to the OnGenerateRowContentWidget slate event. Using default row content.")))
 	{
@@ -532,65 +535,73 @@ void SDataHierarchyEditor::Construct(const FArguments& InArgs, TObjectPtr<UDataH
 			if(HierarchyElement->GetDataMutable()->IsA<UHierarchyCategory>())
 			{
 				TSharedRef<FHierarchyCategoryViewModel> TreeViewCategory = StaticCastSharedRef<FHierarchyCategoryViewModel>(HierarchyElement);
-				return SNew(SHierarchyCategory, TreeViewCategory);
+				return SNew(SHierarchyElement, TreeViewCategory)
+					.Style(&FDataHierarchyEditorStyle::Get().GetWidgetStyle<FInlineEditableTextBlockStyle>("HierarchyEditor.Category"));
+			}
+
+			if(HierarchyElement->GetDataMutable()->IsA<UHierarchyItem>())
+			{
+				TSharedRef<FHierarchyItemViewModel> TreeViewItem = StaticCastSharedRef<FHierarchyItemViewModel>(HierarchyElement);
+				return SNew(SHierarchyElement, TreeViewItem);
 			}
 			
 			return SNew(STextBlock).Text(HierarchyElement->ToStringAsText());
 		});
 	}
+
+	SSplitter::FSlot* SourcePanelSlot = nullptr;
 	SSplitter::FSlot* DetailsPanelSlot = nullptr;
 
-	TSharedRef<SWidget> AddSectionButton = SNew(SHorizontalBox)
-		+ SHorizontalBox::Slot()
-		.AutoWidth()
+	TSharedRef<SWidget> AddSectionButton = SNew(SButton)
+		.OnClicked(this, &SDataHierarchyEditor::OnAddSectionClicked)
+		.ButtonStyle(FDataHierarchyEditorStyle::Get(), "HierarchyEditor.ButtonStyle")
 		[
-			SNew(SButton)
-			.OnClicked(this, &SDataHierarchyEditor::OnAddSectionClicked)
-			.ButtonStyle(FDataHierarchyEditorStyle::Get(), "HierarchyEditor.ButtonStyle")
+			SNew(SHorizontalBox)
+			+ SHorizontalBox::Slot()
+			.AutoWidth()
+			.Padding(2.f)
 			[
-				SNew(SHorizontalBox)
-				+ SHorizontalBox::Slot()
-				.AutoWidth()
-				.Padding(2.f)
-				[
-					SNew(SImage)
-					.Image(FAppStyle::Get().GetBrush("Icons.PlusCircle"))
-				]
-				+ SHorizontalBox::Slot()
-				.VAlign(VAlign_Center)
-				.AutoWidth()
-				[			
-					SNew(STextBlock)
-					.Text(LOCTEXT("AddSectionLabel","Add Section"))				
-				]
+				SNew(SImage)
+				.Image(FAppStyle::Get().GetBrush("Icons.PlusCircle"))
+			]
+			+ SHorizontalBox::Slot()
+			.VAlign(VAlign_Center)
+			.AutoWidth()
+			[			
+				SNew(STextBlock)
+				.Text(LOCTEXT("AddSectionLabel","Add Section"))				
 			]
 		];
 
-	TSharedRef<SWidget> AddCategoryButton = SNew(SHorizontalBox)
-		+ SHorizontalBox::Slot()
-		.AutoWidth()
+	TSharedRef<SWidget> AddCategoryButton = SNew(SButton)
+		.Visibility(HierarchyViewModel->GetCategoryDataClass() ? EVisibility::Visible : EVisibility::Collapsed)
+		.OnClicked(this, &SDataHierarchyEditor::OnAddCategoryClicked)
+		.ButtonStyle(FDataHierarchyEditorStyle::Get(), "HierarchyEditor.ButtonStyle")
 		[
-			SNew(SButton)
-			.OnClicked(this, &SDataHierarchyEditor::OnAddCategoryClicked)
-			.ButtonStyle(FDataHierarchyEditorStyle::Get(), "HierarchyEditor.ButtonStyle")
+			SNew(SHorizontalBox)
+			+ SHorizontalBox::Slot()
+			.AutoWidth()
+			.Padding(2.f)
 			[
-				SNew(SHorizontalBox)
-				+ SHorizontalBox::Slot()
-				.AutoWidth()
-				.Padding(2.f)
-				[
-					SNew(SImage)
-					.Image(FAppStyle::Get().GetBrush("Icons.PlusCircle"))
-				]
-				+ SHorizontalBox::Slot()
-				.VAlign(VAlign_Center)
-				.AutoWidth()
-				[			
-					SNew(STextBlock)
-					.Text(LOCTEXT("AddCategoryLabel","Add Category"))				
-				]
+				SNew(SImage)
+				.Image(FAppStyle::Get().GetBrush("Icons.PlusCircle"))
+			]
+			+ SHorizontalBox::Slot()
+			.VAlign(VAlign_Center)
+			.AutoWidth()
+			[			
+				SNew(STextBlock)
+				.Text(LOCTEXT("AddCategoryLabel","Add Category"))				
 			]
 		];
+
+	TSharedRef<SWidget> AddCustomButton = SNew(SPositiveActionButton)
+		.Visibility_Lambda([this]() -> EVisibility
+		{
+			return HierarchyViewModel->GetAdditionalTypesToAddInUi().Num() > 0 ? EVisibility::Visible : EVisibility::Collapsed;
+		})
+		.Text(LOCTEXT("AddNewButtonLabel", "Add"))
+		.OnGetMenuContent(this, &SDataHierarchyEditor::OnGetCustomAddContent);
 	
 	ChildSlot
 	[
@@ -602,40 +613,9 @@ void SDataHierarchyEditor::Construct(const FArguments& InArgs, TObjectPtr<UDataH
 			.Orientation(Orient_Horizontal)
 			.PhysicalSplitterHandleSize(2.f)
 			+ SSplitter::Slot()
+			.Expose(SourcePanelSlot)
 			.Value(0.3f)
 			.MinSize(0.1f)
-			[
-				SNew(SVerticalBox)
-				+ SVerticalBox::Slot()
-				.AutoHeight()
-				.Padding(2.f)
-				[
-					SAssignNew(SourceSearchBox, SSearchBox)
-					.OnTextChanged(this, &SDataHierarchyEditor::OnSourceSearchTextChanged)
-					.OnTextCommitted(this, &SDataHierarchyEditor::OnSourceSearchTextCommitted)
-					.OnSearch(SSearchBox::FOnSearch::CreateSP(this, &SDataHierarchyEditor::OnSearchButtonClicked))
-					.DelayChangeNotificationsWhileTyping(true)
-					.SearchResultData(this, &SDataHierarchyEditor::GetSearchResultData)
-				]
-				+ SVerticalBox::Slot()
-				.AutoHeight()
-				.Padding(2.f)
-				[
-					SAssignNew(SourceSectionBox, SWrapBox)
-					.UseAllottedSize(true)
-				]
-				+ SVerticalBox::Slot()
-				.Padding(1.f, 2.f)
-				[
-					SAssignNew(SourceTreeView, STreeView<TSharedPtr<FHierarchyElementViewModel>>)
-					.TreeItemsSource(&GetSourceItems())
-					.OnSelectionChanged(this, &SDataHierarchyEditor::OnSelectionChanged, false)
-					.OnGenerateRow(this, &SDataHierarchyEditor::GenerateSourceItemRow)
-					.OnGetChildren_UObject(InHierarchyViewModel.Get(), &UDataHierarchyViewModelBase::OnGetChildren)
-					.OnItemToString_Debug_UObject(InHierarchyViewModel.Get(), &UDataHierarchyViewModelBase::OnElementToStringDebug)
-					.OnContextMenuOpening(this, &SDataHierarchyEditor::SummonContextMenuForSelectedRows, false)
-				]					
-			]
 			+ SSplitter::Slot()
 			.Value(0.4f)
 			.MinSize(0.1f)
@@ -645,12 +625,14 @@ void SDataHierarchyEditor::Construct(const FArguments& InArgs, TObjectPtr<UDataH
 				.AutoHeight()
 				[
 					SNew(SVerticalBox)
+					.Visibility(HierarchyViewModel->GetSectionDataClass() ? EVisibility::Visible : EVisibility::Collapsed)
 					+ SVerticalBox::Slot()
 					.AutoHeight()
 					[
 						SNew(SBorder)
 						.Padding(0.f)
 						.BorderImage(FAppStyle::Get().GetBrush("Brushes.Header"))
+						.HAlign(HAlign_Left)
 						[
 							AddSectionButton
 						]
@@ -673,7 +655,23 @@ void SDataHierarchyEditor::Construct(const FArguments& InArgs, TObjectPtr<UDataH
 						.Padding(0.f)
 						.BorderImage(FAppStyle::Get().GetBrush("Brushes.Header"))
 						[
-							AddCategoryButton
+							SNew(SHorizontalBox)
+							+ SHorizontalBox::Slot()
+							.HAlign(HAlign_Left)
+							.AutoWidth()
+							[
+								AddCategoryButton
+							]
+							+ SHorizontalBox::Slot()
+							[
+								SNew(SSpacer)
+							]
+							+ SHorizontalBox::Slot()
+							.HAlign(HAlign_Right)
+							.AutoWidth()
+							[
+								AddCustomButton
+							]
 						]
 					]
 					+ SVerticalBox::Slot()
@@ -712,7 +710,7 @@ void SDataHierarchyEditor::Construct(const FArguments& InArgs, TObjectPtr<UDataH
 					[
 						SAssignNew(HierarchyTreeView, STreeView<TSharedPtr<FHierarchyElementViewModel>>)
 						.TreeItemsSource(&InHierarchyViewModel->GetHierarchyItems())
-						.OnSelectionChanged(this, &SDataHierarchyEditor::OnSelectionChanged, true)
+						.OnSelectionChanged(this, &SDataHierarchyEditor::OnSelectionChanged)
 						.OnGenerateRow(this, &SDataHierarchyEditor::GenerateHierarchyItemRow)
 						.OnGetChildren_UObject(InHierarchyViewModel.Get(), &UDataHierarchyViewModelBase::OnGetChildren)
 						.OnItemToString_Debug_UObject(InHierarchyViewModel.Get(), &UDataHierarchyViewModelBase::OnElementToStringDebug)
@@ -721,12 +719,54 @@ void SDataHierarchyEditor::Construct(const FArguments& InArgs, TObjectPtr<UDataH
 				]
 			]
 			+ SSplitter::Slot()
+			.Expose(DetailsPanelSlot)		
 			.Value(0.3f)
 			.MinSize(0.1f)
-			.Expose(DetailsPanelSlot)		
 		]
 	];
 
+	if(InHierarchyViewModel->SupportsSourcePanel())
+	{
+		TSharedRef<SWidget> SourcePanelContent = SNew(SVerticalBox)
+		+ SVerticalBox::Slot()
+		.AutoHeight()
+		.Padding(2.f)
+		[
+			SAssignNew(SourceSearchBox, SSearchBox)
+			.OnTextChanged(this, &SDataHierarchyEditor::OnSourceSearchTextChanged)
+			.OnTextCommitted(this, &SDataHierarchyEditor::OnSourceSearchTextCommitted)
+			.OnSearch(SSearchBox::FOnSearch::CreateSP(this, &SDataHierarchyEditor::OnSearchButtonClicked))
+			.DelayChangeNotificationsWhileTyping(true)
+			.SearchResultData(this, &SDataHierarchyEditor::GetSearchResultData)
+		]
+		+ SVerticalBox::Slot()
+		.AutoHeight()
+		.Padding(2.f)
+		[
+			SAssignNew(SourceSectionBox, SWrapBox)
+			.UseAllottedSize(true)
+		]
+		+ SVerticalBox::Slot()
+		.Padding(1.f, 2.f)
+		[
+			SAssignNew(SourceTreeView, STreeView<TSharedPtr<FHierarchyElementViewModel>>)
+			.TreeItemsSource(&GetSourceItems())
+			.OnSelectionChanged(this, &SDataHierarchyEditor::OnSelectionChanged)
+			.OnGenerateRow(this, &SDataHierarchyEditor::GenerateSourceItemRow)
+			.OnGetChildren_UObject(InHierarchyViewModel.Get(), &UDataHierarchyViewModelBase::OnGetChildren)
+			.OnItemToString_Debug_UObject(InHierarchyViewModel.Get(), &UDataHierarchyViewModelBase::OnElementToStringDebug)
+			.OnContextMenuOpening(this, &SDataHierarchyEditor::SummonContextMenuForSelectedRows, false)
+		];
+
+		SourcePanelSlot->AttachWidget(SourcePanelContent);
+	}
+	else
+	{
+		SourcePanelSlot->SetMinSize(0.f);
+		SourcePanelSlot->SetSizeValue(0.f);
+		SourcePanelSlot->SetResizable(false);
+	}
+	
 	if(InHierarchyViewModel->SupportsDetailsPanel())
 	{
 		FPropertyEditorModule& PropertyEditorModule = FModuleManager::Get().LoadModuleChecked<FPropertyEditorModule>("PropertyEditor");
@@ -751,10 +791,14 @@ void SDataHierarchyEditor::Construct(const FArguments& InArgs, TObjectPtr<UDataH
 		{
 			DetailsPanel->RegisterInstancedCustomPropertyLayout(Customizations.Key, Customizations.Value);
 		}
+
+		OnSelectionChanged(nullptr, ESelectInfo::Type::Direct);
 		
 		DetailsPanelSlot->AttachWidget(DetailsPanel.ToSharedRef());
 	}
-		
+
+	// For shortcuts, we use the commandlist logic which lacks context information (such as 'what IS the selected item')
+	
 	HierarchyViewModel->GetCommands()->MapAction(FGenericCommands::Get().Rename,
 		FExecuteAction::CreateSP(this, &SDataHierarchyEditor::RequestRenameSelectedItem),
 		FCanExecuteAction::CreateSP(this, &SDataHierarchyEditor::CanRequestRenameSelectedItem),
@@ -801,21 +845,27 @@ SDataHierarchyEditor::~SDataHierarchyEditor()
 		HierarchyViewModel->GetCommands()->UnmapAction(FGenericCommands::Get().Rename);
 		HierarchyViewModel->GetCommands()->UnmapAction(FDataHierarchyEditorCommands::Get().FindInHierarchy);
 	}
-	
-	SourceRootViewModel->OnSyncPropagated().Unbind();
-	SourceRootViewModel->OnSectionsChanged().Unbind();
-	SourceRootViewModel.Reset();
-	SourceRoot->ConditionalBeginDestroy();
-	SourceRoot = nullptr;
+
+	if(HierarchyViewModel->SupportsSourcePanel())
+	{
+		SourceRootViewModel->OnSyncPropagated().Unbind();
+		SourceRootViewModel->OnSectionsChanged().Unbind();
+		SourceRootViewModel.Reset();
+		SourceRoot->ConditionalBeginDestroy();
+		SourceRoot = nullptr;
+	}
 }
 
 void SDataHierarchyEditor::RefreshSourceItems()
 {
-	SourceRoot->EmptyAllData();
-	HierarchyViewModel->PrepareSourceItems(SourceRoot.Get(), SourceRootViewModel);
-	SourceRootViewModel->SyncViewModelsToData();
-	RefreshSourceView(false);
-	RefreshSectionsView();
+	if(HierarchyViewModel->SupportsSourcePanel())
+	{
+		SourceRoot->EmptyAllData();
+		HierarchyViewModel->PrepareSourceItems(SourceRoot.Get(), SourceRootViewModel);
+		SourceRootViewModel->SyncViewModelsToData();
+		RefreshSourceView(false);
+		RefreshSectionsView();
+	}
 }
 
 void SDataHierarchyEditor::RefreshAllViews(bool bFullRefresh)
@@ -901,6 +951,11 @@ TSharedPtr<SWidget> SDataHierarchyEditor::SummonContextMenuForSelectedRows(bool
 
 void SDataHierarchyEditor::RefreshSourceView(bool bFullRefresh) const
 {
+	if(HierarchyViewModel->SupportsSourcePanel() == false)
+	{
+		return;
+	}
+	
 	SourceTreeView->SetTreeItemsSource(&GetSourceItems());
 	if(bFullRefresh)
 	{
@@ -914,6 +969,11 @@ void SDataHierarchyEditor::RefreshSourceView(bool bFullRefresh) const
 
 void SDataHierarchyEditor::RequestRefreshSourceViewNextFrame(bool bFullRefresh)
 {
+	if(HierarchyViewModel->SupportsSourcePanel() == false)
+	{
+		return;
+	}
+	
 	if(!RefreshSourceViewNextFrameHandle.IsValid())
 	{
 		RefreshSourceViewNextFrameHandle = RegisterActiveTimer(0.f, FWidgetActiveTimerDelegate::CreateLambda([this, bFullRefresh](double CurrentTime, float DeltaTime)
@@ -954,56 +1014,57 @@ void SDataHierarchyEditor::RequestRefreshHierarchyViewNextFrame(bool bFullRefres
 
 void SDataHierarchyEditor::RefreshSectionsView()
 {
-	SourceSectionBox->ClearChildren();
-	HierarchySectionBox->ClearChildren();
-
-	for (TSharedPtr<FHierarchySectionViewModel>& SourceSection : SourceRootViewModel->GetSectionViewModels())
+	if(HierarchyViewModel->SupportsSourcePanel())
 	{
-		TSharedPtr<SHierarchySection> SectionWidget = SNew(SHierarchySection, SourceSection)
-		.IsSectionActive_Lambda([this, SourceSection]()
-		{
-			return GetActiveSourceSection() == SourceSection ? ECheckBoxState::Checked : ECheckBoxState::Unchecked;
-		})
-		.OnSectionActivated_Lambda([this](TSharedPtr<FHierarchySectionViewModel> SectionViewModel)
-		{
-			SetActiveSourceSection(SectionViewModel);
-		});
-		
-		SourceSectionBox->AddSlot()
-		.Padding(2.f)
-		.HAlign(HAlign_Center)
-		.VAlign(VAlign_Center)
-		[
-			SectionWidget.ToSharedRef()		
-		];
-	}
+		SourceSectionBox->ClearChildren();
 
-	if(SourceRootViewModel->GetSectionViewModels().Num() > 0)
-	{
-		DefaultSourceSectionViewModel = MakeShared<FHierarchySectionViewModel>(nullptr, SourceRootViewModel.ToSharedRef(), HierarchyViewModel);
-		
-		TSharedPtr<SHierarchySection> DefaultSourceSection = SNew(SHierarchySection, DefaultSourceSectionViewModel)
-		.IsSectionActive_Lambda([this]()
+		for (TSharedPtr<FHierarchySectionViewModel>& SourceSection : SourceRootViewModel->GetSectionViewModels())
 		{
-			return GetActiveSourceSection() == DefaultSourceSectionViewModel ? ECheckBoxState::Checked : ECheckBoxState::Unchecked;
-		})
-		.OnSectionActivated(this, &SDataHierarchyEditor::SetActiveSourceSection);
+			TSharedPtr<SHierarchySection> SectionWidget = SNew(SHierarchySection, SourceSection)
+			.IsSectionActive_Lambda([this, SourceSection]()
+			{
+				return GetActiveSourceSection() == SourceSection ? ECheckBoxState::Checked : ECheckBoxState::Unchecked;
+			})
+			.OnSectionActivated_Lambda([this](TSharedPtr<FHierarchySectionViewModel> SectionViewModel)
+			{
+				SetActiveSourceSection(SectionViewModel);
+			});
+			
+			SourceSectionBox->AddSlot()
+			.Padding(2.f)
+			.HAlign(HAlign_Center)
+			.VAlign(VAlign_Center)
+			[
+				SectionWidget.ToSharedRef()		
+			];
+		}
 
-		SourceSectionBox->AddSlot()
-		.Padding(2.f)
-		.HAlign(HAlign_Center)
-		.VAlign(VAlign_Center)
-		[
-			DefaultSourceSection.ToSharedRef()
-		];
+		if(SourceRootViewModel->GetSectionViewModels().Num() > 0)
+		{			
+			TSharedPtr<SHierarchySection> DefaultSourceSection = SNew(SHierarchySection, DefaultSourceSectionViewModel)
+			.IsSectionActive_Lambda([this]()
+			{
+				return GetActiveSourceSection() == DefaultSourceSectionViewModel ? ECheckBoxState::Checked : ECheckBoxState::Unchecked;
+			})
+			.OnSectionActivated(this, &SDataHierarchyEditor::SetActiveSourceSection);
+
+			SourceSectionBox->AddSlot()
+			.Padding(2.f)
+			.HAlign(HAlign_Center)
+			.VAlign(VAlign_Center)
+			[
+				DefaultSourceSection.ToSharedRef()
+			];
+		}
 	}
-	
+
+	HierarchySectionBox->ClearChildren();
 	for (TSharedPtr<FHierarchySectionViewModel>& HierarchySection : HierarchyViewModel->GetHierarchyRootViewModel()->GetSectionViewModels())
 	{
 		TSharedPtr<SHierarchySection> SectionWidget = SNew(SHierarchySection, HierarchySection)
 		.IsSectionActive_Lambda([this, HierarchySection]()
 		{
-			return HierarchyViewModel->GetActiveHierarchySection() == HierarchySection ? ECheckBoxState::Checked : ECheckBoxState::Unchecked;
+			return HierarchyViewModel->GetActiveHierarchySectionViewModel() == HierarchySection ? ECheckBoxState::Checked : ECheckBoxState::Unchecked;
 		})
 		.OnSectionActivated_Lambda([this](TSharedPtr<FHierarchySectionViewModel> SectionViewModel)
 		{
@@ -1022,7 +1083,7 @@ void SDataHierarchyEditor::RefreshSectionsView()
 	TSharedPtr<SHierarchySection> DefaultHierarchySection = SNew(SHierarchySection, HierarchyViewModel->GetDefaultHierarchySectionViewModel())
 	.IsSectionActive_Lambda([this]()
 	{
-		return HierarchyViewModel->GetActiveHierarchySection() == HierarchyViewModel->GetDefaultHierarchySectionViewModel() ? ECheckBoxState::Checked : ECheckBoxState::Unchecked;
+		return HierarchyViewModel->GetActiveHierarchySectionViewModel() == HierarchyViewModel->GetDefaultHierarchySectionViewModel() ? ECheckBoxState::Checked : ECheckBoxState::Unchecked;
 	})
 	.OnSectionActivated_Lambda([this](TSharedPtr<FHierarchySectionViewModel> SectionViewModel)
 	{
@@ -1084,7 +1145,7 @@ bool SDataHierarchyEditor::IsItemSelected(TSharedPtr<FHierarchyElementViewModel>
 TSharedRef<ITableRow> SDataHierarchyEditor::GenerateSourceItemRow(TSharedPtr<FHierarchyElementViewModel> HierarchyItem, const TSharedRef<STableViewBase>& TableViewBase)
 {
 	return SNew(STableRow<TSharedPtr<FHierarchyElementViewModel>>, TableViewBase)
-	.Style(HierarchyItem->GetData()->IsA<UHierarchyItem>() ? ItemRowStyle : CategoryRowStyle)
+	.Style(HierarchyItem->GetRowStyle() ? HierarchyItem->GetRowStyle() : &FCoreStyle::Get().GetWidgetStyle<FTableRowStyle>("TableView.Row"))
 	.OnDragDetected(HierarchyItem.ToSharedRef(), &FHierarchyElementViewModel::OnDragDetected, true)
 	.Padding(FMargin(2))
 	[
@@ -1135,7 +1196,7 @@ TSharedRef<ITableRow> SDataHierarchyEditor::GenerateSourceItemRow(TSharedPtr<FHi
 TSharedRef<ITableRow> SDataHierarchyEditor::GenerateHierarchyItemRow(TSharedPtr<FHierarchyElementViewModel> HierarchyItem, const TSharedRef<STableViewBase>& TableViewBase)
 {
 	return SNew(STableRow<TSharedPtr<FHierarchyElementViewModel>>, TableViewBase)
-	.Style(HierarchyItem->GetData()->IsA<UHierarchyItem>() ? ItemRowStyle : CategoryRowStyle)
+	.Style(HierarchyItem->GetRowStyle() ? HierarchyItem->GetRowStyle() : &FCoreStyle::Get().GetWidgetStyle<FTableRowStyle>("TableView.Row"))
 	.OnAcceptDrop(HierarchyItem.ToSharedRef(), &FHierarchyElementViewModel::OnDroppedOnRow)
 	.OnCanAcceptDrop(HierarchyItem.ToSharedRef(), &FHierarchyElementViewModel::OnCanRowAcceptDrop)
 	.OnDragDetected(HierarchyItem.ToSharedRef(), &FHierarchyElementViewModel::OnDragDetected, false)
@@ -1154,17 +1215,17 @@ TSharedRef<ITableRow> SDataHierarchyEditor::GenerateHierarchyItemRow(TSharedPtr<
 			.WidthOverride(10.f)
 			.Visibility_Lambda([HierarchyItem]()
 			{
-				return HierarchyItem->IsEditableByUser().bCanPerform ? EVisibility::Collapsed : EVisibility::Visible;
+				return HierarchyItem->IsEditableByUser().bResult ? EVisibility::Collapsed : EVisibility::Visible;
 			})
 			[
 				SNew(SImage)
 				.Image(FAppStyle::GetBrush("Icons.Lock"))
 				.ToolTipText_Lambda([HierarchyItem]()
 				{
-					FHierarchyElementViewModel::FCanPerformActionResults IsEditableResults = HierarchyItem->IsEditableByUser();
-					if(IsEditableResults.bCanPerform == false)
+					FHierarchyElementViewModel::FResultWithUserFeedback IsEditableResults = HierarchyItem->IsEditableByUser();
+					if(IsEditableResults.bResult == false)
 					{
-						return IsEditableResults.CanPerformMessage;
+						return IsEditableResults.UserFeedback.Get(FText::GetEmpty());
 					}
 
 					return FText::GetEmpty();
@@ -1182,8 +1243,8 @@ bool SDataHierarchyEditor::FilterForSourceSection(TSharedPtr<const FHierarchyEle
 {
 	if(ActiveSourceSection.IsValid())
 	{
-		// If the currently selected section data is nullptr, it's the All section, and we let everything pass
-		if(ActiveSourceSection.Pin()->GetData() == nullptr)
+		// If the currently selected section data is the all section we let everything pass
+		if(ActiveSourceSection.Pin()->IsAllSection())
 		{
 			return true;
 		}
@@ -1229,7 +1290,7 @@ const TArray<TSharedPtr<FHierarchyElementViewModel>>& SDataHierarchyEditor::GetS
 
 bool SDataHierarchyEditor::IsDetailsPanelEditingAllowed() const
 {
-	return SelectedDetailsPanelItemViewModel.IsValid() && SelectedDetailsPanelItemViewModel.Pin()->IsEditableByUser().bCanPerform;
+	return SelectedDetailsPanelItemViewModel.IsValid() && SelectedDetailsPanelItemViewModel.Pin()->IsEditableByUser().bResult;
 }
 
 void SDataHierarchyEditor::RequestRenameSelectedItem()
@@ -1238,7 +1299,7 @@ void SDataHierarchyEditor::RequestRenameSelectedItem()
 
 	if(SelectedItems.Num() == 0)
 	{
-		TSharedPtr<FHierarchySectionViewModel> ActiveSection = HierarchyViewModel->GetActiveHierarchySection();
+		TSharedPtr<FHierarchySectionViewModel> ActiveSection = HierarchyViewModel->GetActiveHierarchySectionViewModel();
 		if(ActiveSection)
 		{
 			SelectedItems = { ActiveSection };
@@ -1257,7 +1318,7 @@ bool SDataHierarchyEditor::CanRequestRenameSelectedItem() const
 
 	if(SelectedItems.Num() == 0)
 	{
-		TSharedPtr<FHierarchySectionViewModel> ActiveSection = HierarchyViewModel->GetActiveHierarchySection();
+		TSharedPtr<FHierarchySectionViewModel> ActiveSection = HierarchyViewModel->GetActiveHierarchySectionViewModel();
 		if(ActiveSection)
 		{
 			SelectedItems = { ActiveSection };
@@ -1274,6 +1335,11 @@ bool SDataHierarchyEditor::CanRequestRenameSelectedItem() const
 
 void SDataHierarchyEditor::ClearSourceItems() const
 {
+	if(HierarchyViewModel->SupportsSourcePanel() == false)
+	{
+		return;
+	}
+	
 	SourceRoot->GetChildrenMutable().Empty();
 	SourceRoot->GetSectionDataMutable().Empty();
 	SourceRootViewModel->GetChildrenMutable().Empty();
@@ -1291,7 +1357,7 @@ void SDataHierarchyEditor::DeleteSelectedHierarchyItems() const
 	
 	if(SelectedElements.Num() == 0)
 	{
-		TSharedPtr<FHierarchySectionViewModel> ActiveSection = HierarchyViewModel->GetActiveHierarchySection();
+		TSharedPtr<FHierarchySectionViewModel> ActiveSection = HierarchyViewModel->GetActiveHierarchySectionViewModel();
 		if(ActiveSection)
 		{
 			SelectedElements = { ActiveSection };
@@ -1307,7 +1373,7 @@ bool SDataHierarchyEditor::CanDeleteSelectedElements() const
 
 	if(SelectedElements.Num() == 0)
 	{
-		TSharedPtr<FHierarchySectionViewModel> ActiveSection = HierarchyViewModel->GetActiveHierarchySection();
+		TSharedPtr<FHierarchySectionViewModel> ActiveSection = HierarchyViewModel->GetActiveHierarchySectionViewModel();
 		if(ActiveSection)
 		{
 			SelectedElements = { ActiveSection };
@@ -1347,6 +1413,11 @@ void SDataHierarchyEditor::NavigateToMatchingHierarchyElementFromSelectedSourceE
 
 bool SDataHierarchyEditor::CanNavigateToMatchingHierarchyElementFromSelectedSourceElement() const
 {
+	if(HierarchyViewModel->SupportsSourcePanel() == false)
+	{
+		return false;
+	}
+	
 	TArray<TSharedPtr<FHierarchyElementViewModel>> SelectedElements = SourceTreeView->GetSelectedItems();
 
 	if(SelectedElements.Num() != 1)
@@ -1368,19 +1439,6 @@ bool SDataHierarchyEditor::CanNavigateToMatchingHierarchyElementFromSelectedSour
 	return false;
 }
 
-void SDataHierarchyEditor::DeleteActiveSection() const
-{
-	if(TSharedPtr<FHierarchySectionViewModel> ActiveHierarchySectionViewModel = HierarchyViewModel->GetActiveHierarchySection())
-	{
-		DeleteItems({ActiveHierarchySectionViewModel});
-	}
-}
-
-bool SDataHierarchyEditor::CanDeleteActiveSection() const
-{
-	return HierarchyViewModel->GetActiveHierarchySection() != nullptr;
-}
-
 void SDataHierarchyEditor::OnElementAdded(TSharedPtr<FHierarchyElementViewModel> AddedItem)
 {
 	// when a new item is created (opposed to dragged & dropped from source view, i.e. only categories so far)
@@ -1402,12 +1460,12 @@ void SDataHierarchyEditor::OnHierarchySectionActivated(TSharedPtr<FHierarchySect
 {
 	// We forward nullptr in case this is the 'All' default section. We determine this by checking its data for validity.
 	// The all section does not have an actual element associated with it.
-	OnSelectionChanged(Section.IsValid() && Section->GetData() ? Section : nullptr, ESelectInfo::Direct, true);
+	OnSelectionChanged(Section, ESelectInfo::Direct);
 }
 
 void SDataHierarchyEditor::OnSourceSectionActivated(TSharedPtr<FHierarchySectionViewModel> Section)
 {
-	OnSelectionChanged(Section.IsValid() && Section->GetData() ? Section : nullptr, ESelectInfo::Direct, false);
+	OnSelectionChanged(Section, ESelectInfo::Direct);
 	RunSourceSearch();
 }
 
@@ -1419,14 +1477,122 @@ void SDataHierarchyEditor::OnHierarchySectionAdded(TSharedPtr<FHierarchySectionV
 
 void SDataHierarchyEditor::OnHierarchySectionDeleted(TSharedPtr<FHierarchySectionViewModel> DeletedSection)
 {
-	if(HierarchyViewModel->GetActiveHierarchySection() == DeletedSection)
+	if(HierarchyViewModel->GetActiveHierarchySectionViewModel() == DeletedSection)
 	{
 		HierarchyViewModel->SetActiveHierarchySection(HierarchyViewModel->GetDefaultHierarchySectionViewModel());
 	}
 }
 
+TSharedRef<SWidget> SDataHierarchyEditor::OnGetCustomAddContent()
+{
+	FToolMenuOwnerScoped OwnerScoped(UE_MODULE_NAME);
+
+	static FName MenuName = "DataHierarchyEditor.CustomAddMenu";
+	if(UToolMenus::Get()->IsMenuRegistered(MenuName) == false)
+	{
+		UToolMenu* ToolMenu = UToolMenus::Get()->RegisterMenu(MenuName);
+		ToolMenu->AddDynamicSection(NAME_None, FNewToolMenuDelegate::CreateLambda([](UToolMenu* DynamicMenu)
+		{
+			UDataHierarchyEditorMenuContext* ContextObject = DynamicMenu->Context.FindContext<UDataHierarchyEditorMenuContext>();
+			if(!ensure(ContextObject))
+			{
+				return;
+			}
+
+			TArray<TSubclassOf<UHierarchyElement>> AdditionalTypes = ContextObject->Widget.Pin()->HierarchyViewModel->GetAdditionalTypesToAddInUi();
+			for(TSubclassOf<UHierarchyElement> Type : AdditionalTypes)
+			{
+				// We use the first available source. There is guaranteed to be one, although there could be multiple.
+				FToolMenuSection& Section = DynamicMenu->FindOrAddSection(NAME_None);
+
+				FToolUIAction UIAction;
+				UIAction.ExecuteAction = FToolMenuExecuteAction::CreateLambda([ContextObject, Type](const FToolMenuContext& Context)
+				{
+					TSharedPtr<SDataHierarchyEditor> DataHierarchyEditor = ContextObject->Widget.Pin();
+					TArray<TSharedPtr<FHierarchyElementViewModel>> SelectedItems;
+					DataHierarchyEditor->HierarchyTreeView->GetSelectedItems(SelectedItems);
+
+					TSharedPtr<FHierarchyElementViewModel> ParentViewModel;
+					if(SelectedItems.Num() == 1)
+					{
+						ParentViewModel = SelectedItems[0];
+					}
+					else
+					{
+						ParentViewModel = DataHierarchyEditor->HierarchyViewModel->GetHierarchyRootViewModel();
+					}
+
+					if(ParentViewModel->CanContain(Type).bResult)
+					{
+						ParentViewModel->AddChild(Type, FHierarchyElementIdentity());
+					}
+				});
+				UIAction.CanExecuteAction = FToolMenuCanExecuteAction::CreateLambda([ContextObject, Type](const FToolMenuContext& Context)
+				{
+					TSharedPtr<SDataHierarchyEditor> DataHierarchyEditor = ContextObject->Widget.Pin();
+					TArray<TSharedPtr<FHierarchyElementViewModel>> SelectedItems;
+					DataHierarchyEditor->HierarchyTreeView->GetSelectedItems(SelectedItems);
+
+					TSharedPtr<FHierarchyElementViewModel> ParentViewModel;
+					if(SelectedItems.Num() == 1)
+					{
+						ParentViewModel = SelectedItems[0];
+					}
+					else
+					{
+						ParentViewModel = DataHierarchyEditor->HierarchyViewModel->GetHierarchyRootViewModel();
+					}
+					
+					return ParentViewModel->CanContain(Type).bResult;
+				});
+
+				TAttribute<FText> TooltipAttribute = TAttribute<FText>::CreateLambda([ContextObject, Type]()
+				{
+					TSharedPtr<SDataHierarchyEditor> DataHierarchyEditor = ContextObject->Widget.Pin();
+					TArray<TSharedPtr<FHierarchyElementViewModel>> SelectedItems;
+					DataHierarchyEditor->HierarchyTreeView->GetSelectedItems(SelectedItems);
+
+					TSharedPtr<FHierarchyElementViewModel> ParentViewModel;
+
+					if(SelectedItems.Num() == 1)
+					{
+						ParentViewModel = SelectedItems[0];
+					}
+					else
+					{
+						ParentViewModel = DataHierarchyEditor->HierarchyViewModel->GetHierarchyRootViewModel();
+					}
+					
+					if(TOptional<FText> DynamicText = ParentViewModel->CanContain(Type).UserFeedback)
+					{
+						FText BaseText = FText::AsCultureInvariant("{0}\n\n{1}");
+						FText Text = DynamicText.GetValue();
+						return FText::FormatOrdered(BaseText, Type->GetToolTipText(), Text);
+					}
+
+					return Type->GetToolTipText();
+				});
+				
+				FToolMenuEntry Entry = FToolMenuEntry::InitMenuEntry(Type->GetFName(), Type->GetDisplayNameText(), TooltipAttribute, FSlateIconFinder::FindIconForClass(Type), UIAction);
+				Section.AddEntry(Entry);
+			}
+		}));
+	}
+
+	UDataHierarchyEditorMenuContext* MenuContext = NewObject<UDataHierarchyEditorMenuContext>();
+	MenuContext->Widget = SharedThis(this);
+	
+	return UToolMenus::Get()->GenerateWidget(MenuName, FToolMenuContext(MenuContext));
+
+}
+
 void SDataHierarchyEditor::SetActiveSourceSection(TSharedPtr<FHierarchySectionViewModel> Section)
 {
+	if(HierarchyViewModel->SupportsSourcePanel() == false)
+	{
+		return;
+	}
+	
 	ActiveSourceSection = Section;	
 	RefreshSourceView(true);
 	OnSourceSectionActivated(Section);
@@ -1442,36 +1608,57 @@ UHierarchySection* SDataHierarchyEditor::GetActiveSourceSectionData() const
 	return ActiveSourceSection.IsValid() ? ActiveSourceSection.Pin()->GetDataMutable<UHierarchySection>() : nullptr;
 }
 
-void SDataHierarchyEditor::OnSelectionChanged(TSharedPtr<FHierarchyElementViewModel> HierarchyItem, ESelectInfo::Type Type, bool bFromHierarchy) const
+void SDataHierarchyEditor::OnSelectionChanged(TSharedPtr<FHierarchyElementViewModel> SelectedHierarchyElement, ESelectInfo::Type Type) const
 {
 	SelectedDetailsPanelItemViewModel.Reset();
 	if(DetailsPanel.IsValid())
 	{
-		if(HierarchyItem.IsValid() && HierarchyItem->AllowEditingInDetailsPanel())
+		TSharedPtr<FHierarchyElementViewModel> ElementToDisplay = SelectedHierarchyElement;
+
+		// If the selected element is invalid and we want to select the root instead of nothing, we do so
+		if((ElementToDisplay.IsValid() == false || ElementToDisplay->AllowEditingInDetailsPanel() == false))
+		{
+			if(HierarchyViewModel->SelectRootInsteadOfNone())
+			{
+				ElementToDisplay = HierarchyViewModel->GetHierarchyRootViewModel();
+			}
+			else
+			{
+				ElementToDisplay = nullptr;
+			}
+		}
+		
+		if(ElementToDisplay.IsValid())
 		{
 			// when se select a section, and the previous item selection is no longer available due to it, we would get a selection refresh next tick
 			// to wipe out the current selection. We want to avoid that, so we manually clear the selected items in that case.
-			if(HierarchyItem->GetData()->IsA<UHierarchySection>())
+			if(ElementToDisplay->GetData()->IsA<UHierarchySection>())
 			{
 				HierarchyTreeView->ClearSelection();
 			}
 
 			// we clear the selection of the other tree view to avoid 
-			if(bFromHierarchy)
+			if(ElementToDisplay->IsForHierarchy())
 			{
-				SourceTreeView->ClearSelection();
+				if(HierarchyViewModel->SupportsSourcePanel())
+				{
+					SourceTreeView->ClearSelection();
+				}
 			}
 			else
 			{
 				HierarchyTreeView->ClearSelection();
 			}
 			
-			UObject* DataForEditing = HierarchyItem->GetDataForEditing();
-			DataForEditing->SetFlags(RF_Transactional);
+			UObject* ObjectForEditing = ElementToDisplay->GetObjectForEditing();
+			if(ObjectForEditing)
+			{
+				ObjectForEditing->SetFlags(RF_Transactional);
+			}
 
 			// we make sure the object we are editing is transactional
-			DetailsPanel->SetObject(DataForEditing);
-			SelectedDetailsPanelItemViewModel = HierarchyItem;
+			DetailsPanel->SetObject(ObjectForEditing);
+			SelectedDetailsPanelItemViewModel = ElementToDisplay;
 		}
 		else
 		{
@@ -1491,8 +1678,8 @@ void SDataHierarchyEditor::OnSelectionChanged(TSharedPtr<FHierarchyElementViewMo
 		[
 			SNew(SImage)
 			.Image(FAppStyle::GetBrush("Icons.Lock"))
-			.Visibility(SelectedDetailsPanelItemViewModel.Pin()->IsEditableByUser().bCanPerform ? EVisibility::Collapsed : EVisibility::Visible)
-			.ToolTipText(SelectedDetailsPanelItemViewModel.Pin()->IsEditableByUser().CanPerformMessage)
+			.Visibility(SelectedDetailsPanelItemViewModel.Pin()->IsEditableByUser().bResult ? EVisibility::Collapsed : EVisibility::Visible)
+			.ToolTipText(SelectedDetailsPanelItemViewModel.Pin()->IsEditableByUser().UserFeedback.Get(FText::GetEmpty()))
 		]
 		+ SHorizontalBox::Slot()
 		.AutoWidth()
@@ -1690,7 +1877,7 @@ FReply SDataHierarchyEditor::HandleHierarchyRootDrop(const FGeometry& Geometry,
 	return FReply::Unhandled();
 }
 
-FHierarchyElementViewModel::FCanPerformActionResults SDataHierarchyEditor::CanDropOnRoot(TSharedPtr<FHierarchyElementViewModel> DraggedItem) const
+FHierarchyElementViewModel::FResultWithUserFeedback SDataHierarchyEditor::CanDropOnRoot(TSharedPtr<FHierarchyElementViewModel> DraggedItem) const
 {
 	return HierarchyViewModel->GetHierarchyRootViewModel()->CanDropOnInternal(DraggedItem, EItemDropZone::OntoItem);
 }
@@ -1701,7 +1888,7 @@ bool SDataHierarchyEditor::OnCanDropOnRoot(TSharedPtr<FDragDropOperation> DragDr
 	{
 		TSharedPtr<FHierarchyDragDropOp> HierarchyDragDropOp = StaticCastSharedPtr<FHierarchyDragDropOp>(DragDropOperation);
 		TSharedPtr<FHierarchyElementViewModel> DraggedItem = HierarchyDragDropOp->GetDraggedElement().Pin();
-		return CanDropOnRoot(DraggedItem).bCanPerform;
+		return CanDropOnRoot(DraggedItem).bResult;
 	}
 	
 	return false;
@@ -1711,8 +1898,8 @@ void SDataHierarchyEditor::OnRootDragEnter(const FDragDropEvent& DragDropEvent)
 {
 	if(TSharedPtr<FHierarchyDragDropOp> HierarchyDragDropOp = DragDropEvent.GetOperationAs<FHierarchyDragDropOp>())
 	{
-		FHierarchyElementViewModel::FCanPerformActionResults CanPerformActionResults = CanDropOnRoot(HierarchyDragDropOp->GetDraggedElement().Pin());
-		HierarchyDragDropOp->SetDescription(CanPerformActionResults.CanPerformMessage);
+		FHierarchyElementViewModel::FResultWithUserFeedback CanPerformActionResults = CanDropOnRoot(HierarchyDragDropOp->GetDraggedElement().Pin());
+		HierarchyDragDropOp->SetDescription(CanPerformActionResults.UserFeedback.Get(FText::GetEmpty()));
 	}
 }
 
@@ -1731,7 +1918,7 @@ FSlateColor SDataHierarchyEditor::GetRootIconColor() const
 		if(FSlateApplication::Get().GetDragDroppingContent()->IsOfType<FHierarchyDragDropOp>())
 		{
 			TSharedPtr<FHierarchyDragDropOp> HierarchyDragDropOp = StaticCastSharedPtr<FHierarchyDragDropOp>(FSlateApplication::Get().GetDragDroppingContent());
-			if(CanDropOnRoot(HierarchyDragDropOp->GetDraggedElement().Pin()).bCanPerform)
+			if(CanDropOnRoot(HierarchyDragDropOp->GetDraggedElement().Pin()).bResult)
 			{
 				return FLinearColor(0.8f, 0.8f, 0.8f, 0.8f);
 			}
diff --git a/Engine/Source/Editor/DataHierarchyEditor/Public/DataHierarchyViewModelBase.h b/Engine/Source/Editor/DataHierarchyEditor/Public/DataHierarchyViewModelBase.h
--- a/Engine/Source/Editor/DataHierarchyEditor/Public/DataHierarchyViewModelBase.h
+++ b/Engine/Source/Editor/DataHierarchyEditor/Public/DataHierarchyViewModelBase.h
@@ -170,7 +170,7 @@ public:
 	bool DoesOneChildExist(bool bRecursive = false) const;
 
 	template<class ChildClass>
-	TArray<ChildClass*> GetChildrenOfType(TArray<ChildClass*>& Out, bool bRecursive = false) const;
+	void GetChildrenOfType(TArray<ChildClass*>& Out, bool bRecursive = false) const;
 
 	template<class PREDICATE_CLASS>
 	void SortChildren(const PREDICATE_CLASS& Predicate, bool bRecursive = false);
@@ -182,6 +182,7 @@ public:
 	void SetIdentity(FHierarchyElementIdentity InIdentity) { Identity = InIdentity; }
 	FHierarchyElementIdentity GetPersistentIdentity() const { return Identity; }
 	UE_API TArray<FHierarchyElementIdentity> GetParentIdentities() const;
+	UE_API TArray<const TObjectPtr<UHierarchyElement>> GetParentElements() const;
 	
 	/** Overridden modify method to also mark all children as modified */
 	UE_API virtual bool Modify(bool bAlwaysMarkDirty = true) override;
@@ -235,7 +236,7 @@ bool UHierarchyElement::DoesOneChildExist(bool bRecursive) const
 }
 
 template <class ChildClass>
-TArray<ChildClass*> UHierarchyElement::GetChildrenOfType(TArray<ChildClass*>& Out, bool bRecursive) const
+void UHierarchyElement::GetChildrenOfType(TArray<ChildClass*>& Out, bool bRecursive) const
 {
 	for(UHierarchyElement* ChildElement : Children)
 	{
@@ -252,8 +253,6 @@ TArray<ChildClass*> UHierarchyElement::GetChildrenOfType(TArray<ChildClass*>& Ou
 			ChildElement->GetChildrenOfType<ChildClass>(Out, bRecursive);
 		}
 	}
-
-	return Out;
 }
 
 template <class PREDICATE_CLASS>
@@ -465,6 +464,7 @@ public:
 	UE_API TSharedPtr<FHierarchyElementViewModel> CreateViewModelForElement(UHierarchyElement* Element, TSharedPtr<FHierarchyElementViewModel> Parent);
 	/** Get the root view model associated with the hierarchy. */
 	TSharedPtr<struct FHierarchyRootViewModel> GetHierarchyRootViewModel() const { return HierarchyRootViewModel; }
+	
 	/** Hierarchy items reflect the already edited hierarchy. This should generally be constructed from persistent serialized data. */
 	UE_API const TArray<TSharedPtr<FHierarchyElementViewModel>>& GetHierarchyItems() const;
 
@@ -472,19 +472,33 @@ public:
 public:
 	/** The hierarchy root the widget is editing. This should point to persistent data stored somewhere else as the serialized root of the hierarchy. */
 	virtual UHierarchyRoot* GetHierarchyRoot() const PURE_VIRTUAL(UDataHierarchyViewModelBase::GetHierarchyRoot, return nullptr;);
+	
+	/** Prepares the items we want to create a hierarchy for. Primary purpose is to add children to the source root to gather the items to display in the source panel.
+	 * The root view model is also given as a way to forcefully sync view models to access additional functionality, if needed. Only used if UsesSourceItems is true. */
+	virtual void PrepareSourceItems(UHierarchyRoot* SourceRoot, TSharedPtr<FHierarchyRootViewModel> SourceRootViewModel) {};
+
+	/** If specified, will add a button in the Data Hierarchy Editor that allows adding the given types from a menu rather than drag & dropping them into the hierarchy. */
+	virtual TArray<TSubclassOf<UHierarchyElement>> GetAdditionalTypesToAddInUi() const { return TArray<TSubclassOf<UHierarchyElement>>(); }
+	
 	/** The outer for the transient source root creation can be overridden. */
 	virtual UObject* GetOuterForSourceRoot() const { return GetTransientPackage(); }
-	/** Prepares the items we want to create a hierarchy for. Primary purpose is to add children to the source root to gather the items to display in the source panel.
-	 * The root view model is also given as a way to forcefully sync view models to access additional functionality, if needed */
-	virtual void PrepareSourceItems(UHierarchyRoot* SourceRoot, TSharedPtr<FHierarchyRootViewModel> SourceRootViewModel) PURE_VIRTUAL(UDataHierarchyViewModelBase::PrepareSourceItems,);
+	
 	/** The class used for creating categories. You can subclass UHierarchyCategory to add new properties. */
 	UE_API virtual TSubclassOf<UHierarchyCategory> GetCategoryDataClass() const;
+	
 	/** The class used for creating sections. You can subclass UHierarchySection to add new properties. */
 	UE_API virtual TSubclassOf<UHierarchySection> GetSectionDataClass() const;
 	/** Function to implement drag drop ops. FHierarchyDragDropOp is a default implementation for a single hierarchy element. Inherit from it and override CreateCustomDecorator for custom UI. */
 	UE_API virtual TSharedRef<FHierarchyDragDropOp> CreateDragDropOp(TSharedRef<FHierarchyElementViewModel> Item);
 	/** This needs to return true if you want the details panel to show up. */
 	virtual bool SupportsDetailsPanel() { return true; }
+
+	/** If false, disables source item functionality (left panel). */
+	virtual bool SupportsSourcePanel() const { return true; }
+
+	/** Attempts to select the root object when no object is selected. Can be used in conjunction with `GetObjectForEditing' to point at external objects. */
+	virtual bool SelectRootInsteadOfNone() const { return false; }
+	
 	/** Overriding this allows to define details panel instance customizations for specific UClasses */
 	virtual TArray<TTuple<UClass*, FOnGetDetailCustomizationInstance>> GetInstanceCustomizations() { return {}; }
 
@@ -550,7 +564,7 @@ public:
 	
 	// Sections
 	UE_API void SetActiveHierarchySection(TSharedPtr<struct FHierarchySectionViewModel>);
-	UE_API TSharedPtr<FHierarchySectionViewModel> GetActiveHierarchySection() const;
+	UE_API TSharedPtr<FHierarchySectionViewModel> GetActiveHierarchySectionViewModel() const;
 	UE_API UHierarchySection* GetActiveHierarchySectionData() const;
 	UE_API bool IsHierarchySectionActive(const UHierarchySection* Section) const;
 	FOnSectionActivated& OnHierarchySectionActivated() { return OnHierarchySectionActivatedDelegate; }
@@ -582,7 +596,7 @@ protected:
 	TSharedPtr<struct FHierarchyRootViewModel> HierarchyRootViewModel;
 
 	TSharedPtr<FHierarchySectionViewModel> DefaultHierarchySectionViewModel;
-	TWeakPtr<struct FHierarchySectionViewModel> ActiveHierarchySection;
+	TWeakPtr<struct FHierarchySectionViewModel> ActiveHierarchySectionViewModel;
 
 	TSharedPtr<FUICommandList> Commands;
 
@@ -615,6 +629,14 @@ private:
 	
 	UPROPERTY(Transient)
 	bool bIsFinalized = false;
+
+	/** We construct a transient all section that isn't actually contained within the hierarchy. */
+	UPROPERTY(Transient)
+	TObjectPtr<UHierarchySection> AllSection;
+public:
+	static FName AllSectionHierarchyObjectName;
+	static FName AllSectionSourceObjectName;
+	static FText AllSectionDisplayName;
 };
 
 /** The base view model for all elements in the hierarchy. There are four base view models inheriting from this; for roots, items, categories, and sections.
@@ -697,17 +719,17 @@ struct FHierarchyElementViewModel : TSharedFromThis<FHierarchyElementViewModel>,
 	UE_API virtual FReply OnDroppedOnRow(const FDragDropEvent& DragDropEvent, EItemDropZone ItemDropZone, TSharedPtr<FHierarchyElementViewModel> Item);
 	UE_API void OnRowDragLeave(const FDragDropEvent& DragDropEvent);
 
-	struct FCanPerformActionResults
+	struct UE_API FResultWithUserFeedback
 	{
-		FCanPerformActionResults(bool bInCanPerform) : bCanPerform(bInCanPerform) {}
+		FResultWithUserFeedback(bool bInCanPerform) : bResult(bInCanPerform) {}
 		
-		bool bCanPerform = false;
-		/** A message that is used when bCanPerform is false. Will either be used in tooltips in the hierarchy editor or as popup message. */
-		FText CanPerformMessage;
+		bool bResult = false;
+		/** A message that can be used for user feedback. Will either be used in tooltips in the hierarchy editor or as popup message etc.. */
+		TOptional<FText> UserFeedback;
 
 		bool operator==(const bool& bOther) const
 		{
-			return bCanPerform == bOther;
+			return bResult == bOther;
 		}
 
 		bool operator!=(const bool& bOther) const
@@ -717,16 +739,16 @@ struct FHierarchyElementViewModel : TSharedFromThis<FHierarchyElementViewModel>,
 	};
 
 	/** Should return true if properties are supposed to be editable & needs to be true if typical operations should work on it (renaming, dragging, deleting etc.) */
-	virtual FCanPerformActionResults IsEditableByUser() { return FCanPerformActionResults(false); }
+	virtual FResultWithUserFeedback IsEditableByUser() { return FResultWithUserFeedback(false); }
 	
 	/** Needs to be true in order to allow drag & drop operations to parent items to this item */
-	virtual bool CanHaveChildren() const { return false; }
+	virtual FResultWithUserFeedback CanHaveChildren() const { return false; }
 	
 	/** Should return true if an item should be draggable. An uneditable item can not be dragged even if CanDragInternal returns true. */
-	UE_API FCanPerformActionResults CanDrag();
+	FResultWithUserFeedback CanDrag();
 	
 	/** Returns true if renamable */
-	bool CanRename() { return IsEditableByUser().bCanPerform && CanRenameInternal(); }
+	bool CanRename() { return IsEditableByUser().bResult && CanRenameInternal(); }
 
 	void Rename(FName NewName) { RenameInternal(NewName); HierarchyViewModel->OnHierarchyPropertiesChanged().Broadcast(); }
 
@@ -748,14 +770,19 @@ struct FHierarchyElementViewModel : TSharedFromThis<FHierarchyElementViewModel>,
 	}
 	
 	/** Returns true if deletable */
-	bool CanDelete() { return IsEditableByUser().bCanPerform && CanDeleteInternal(); }
+	bool CanDelete() { return IsEditableByUser().bResult && CanDeleteInternal(); }
 
 	/** Returns true if the given item can be dropped on the given target area. */
-	FCanPerformActionResults CanDropOn(TSharedPtr<FHierarchyElementViewModel> DraggedItem, EItemDropZone ItemDropZone) { return CanDropOnInternal(DraggedItem, ItemDropZone); }
+	FResultWithUserFeedback CanDropOn(TSharedPtr<FHierarchyElementViewModel> DraggedItem, EItemDropZone ItemDropZone);
 
 	/** Gets executed when an item was dropped on this. */
-	void OnDroppedOn(TSharedPtr<FHierarchyElementViewModel> DroppedItem, EItemDropZone ItemDropZone) { OnDroppedOnInternal(DroppedItem, ItemDropZone); }
-
+	void OnDroppedOn(TSharedPtr<FHierarchyElementViewModel> DroppedItem, EItemDropZone ItemDropZone);
+	/** Lets you completely customize the OnDroppedOn behavior. Needs to return true if you override this. */
+	virtual bool OnDroppedOnOverride(TSharedPtr<FHierarchyElementViewModel> DroppedItem, EItemDropZone ItemDropZone) { return false; }
+	
+	/** Returns true if this element can add a new element via the add menu. */
+	FResultWithUserFeedback CanContain(TSubclassOf<UHierarchyElement> HierarchyElementType);
+	
 	/** Determines the section this item belongs to. */
 	const UHierarchySection* GetSection() const { return GetSectionInternal(); }
 	
@@ -767,7 +794,7 @@ struct FHierarchyElementViewModel : TSharedFromThis<FHierarchyElementViewModel>,
 	virtual bool DoesExternalDataStillExist(const UHierarchyDataRefreshContext* Context) const { return false; }
 
 	/** The UObject we display in the details panel when this item is selected. By default it's the hierarchy element the view model represents. */
-	virtual UObject* GetDataForEditing() { return Element; }
+	virtual UObject* GetObjectForEditing() { return Element; }
 	/** Source items are transient, which is why we don't allow editing by default.
 	 * This is useful to override if source data points at actual data to edit. */
 	virtual bool AllowEditingInDetailsPanel() const { return bIsForHierarchy; }
@@ -785,23 +812,28 @@ struct FHierarchyElementViewModel : TSharedFromThis<FHierarchyElementViewModel>,
 
 	TWeakPtr<FHierarchyElementViewModel> GetParent() { return Parent; }
 
+	virtual const FTableRowStyle* GetRowStyle() const { return nullptr; }
+
 protected:
+	/** Override this to handle post-add logic (for example when dropping an element on another element, or using the add button). */
+	virtual void PostAddFixup() { }
+private:
+	/** Should return true if this element can contain other elements as children. Can optionally specify a message for failure. */
+	UE_API virtual FResultWithUserFeedback CanContainInternal(TSubclassOf<UHierarchyElement> InHierarchyElementType) { return false; }
+	
 	/** Should return true if draggable. An optional message can be provided if false that will show as a slate notification. */
-	virtual FCanPerformActionResults CanDragInternal() { return false; }
+	UE_API virtual FResultWithUserFeedback CanDragInternal() { return false; }
 
 	/** Should return true if renamable */
-	virtual bool CanRenameInternal() { return false; }
+	UE_API virtual bool CanRenameInternal() { return false; }
 
-	virtual void RenameInternal(FName NewName) {}
+	UE_API virtual void RenameInternal(FName NewName) {}
 	
 	/** Should return true if deletable. By default, we can delete items in the hierarchy, not in the source. */
-	virtual bool CanDeleteInternal() { return IsForHierarchy(); }
+	UE_API virtual bool CanDeleteInternal() { return IsForHierarchy(); }
 
 	/** Should return true if the given drag drop operation is allowed to succeed. */
-	UE_API virtual FCanPerformActionResults CanDropOnInternal(TSharedPtr<FHierarchyElementViewModel>, EItemDropZone ItemDropZone);
-	
-	/** Override this to handle drop-on logic. This is called when an item has been dropped on the item that has implemented this function. */
-	virtual void OnDroppedOnInternal(TSharedPtr<FHierarchyElementViewModel>, EItemDropZone ItemDropZone) { }
+	UE_API virtual FResultWithUserFeedback CanDropOnInternal(TSharedPtr<FHierarchyElementViewModel>, EItemDropZone ItemDropZone);
 
 	/** Can be overridden to support sections in the source list.
 	 * In the hierarchy only categories can be parented directly to the root, but using this it is possible to add items to custom sections in the source panel.
@@ -809,13 +841,13 @@ protected:
 	virtual const UHierarchySection* GetSectionInternal() const { return nullptr; }
 private:
 	/** Optionally implement this to refresh dependent data. */
-	virtual void RefreshChildrenDataInternal() {}
+	UE_API virtual void RefreshChildrenDataInternal() {}
 	/** Optionally implement this to further customize the view model sync process.
 	 * An example for this is how the root view model handles sections, as sections exist outside the children hierarchy */
-	virtual void SyncViewModelsToDataInternal() {}
+	UE_API virtual void SyncViewModelsToDataInternal() {}
 	/** Optionally implement this to handle shutdown logic.
 	 * An example for this is when a section gets deleted, it iterates over all categories to null out the associated section  */
-	virtual void FinalizeInternal() {}
+	UE_API virtual void FinalizeInternal() {}
 
 	UE_API void PropagateOnChildSynced();
 protected:
@@ -869,11 +901,12 @@ struct FHierarchyRootViewModel : FHierarchyElementViewModel
 	UE_API virtual ~FHierarchyRootViewModel() override;
 
 	UE_API virtual void Initialize() override;
-	
-	UE_API virtual FCanPerformActionResults CanDropOnInternal(TSharedPtr<FHierarchyElementViewModel> DraggedElement, EItemDropZone ItemDropZone) override;
-	UE_API virtual void OnDroppedOnInternal(TSharedPtr<FHierarchyElementViewModel> DroppedItem, EItemDropZone ItemDropZone) override;
 
-	virtual bool CanHaveChildren() const override { return true; }
+	UE_API virtual FResultWithUserFeedback IsEditableByUser() override { return true; }
+	UE_API virtual FResultWithUserFeedback CanContainInternal(TSubclassOf<UHierarchyElement> InHierarchyElement) override;
+	UE_API virtual FResultWithUserFeedback CanDropOnInternal(TSharedPtr<FHierarchyElementViewModel> DraggedElement, EItemDropZone ItemDropZone) override;
+
+	virtual FResultWithUserFeedback CanHaveChildren() const override { return true; }
 
 	UE_API TSharedPtr<struct FHierarchySectionViewModel> AddSection();
 	UE_API void DeleteSection(TSharedPtr<FHierarchyElementViewModel> SectionViewModel);
@@ -900,16 +933,13 @@ private:
 
 struct FHierarchySectionViewModel : FHierarchyElementViewModel
 {
-	FHierarchySectionViewModel(UHierarchySection* InItem, TSharedRef<FHierarchyRootViewModel> InParent, TWeakObjectPtr<UDataHierarchyViewModelBase> InHierarchyViewModel) : FHierarchyElementViewModel(InItem, InParent, InHierarchyViewModel, InParent->IsForHierarchy())
-	{
-		if(bIsForHierarchy == false)
-		{
-			SetDropDisallowed(true);
-		}
-	}
+	FHierarchySectionViewModel(UHierarchySection* InItem, TSharedRef<FHierarchyRootViewModel> InParent, TWeakObjectPtr<UDataHierarchyViewModelBase> InHierarchyViewModel);
 	
 	virtual ~FHierarchySectionViewModel() override {}
 
+	/** A custom function to retrieve data; if it's the all section, we want to return nullptr as that's what it serializes as. */
+	UHierarchySection* GetSectionData() const;
+	
 	UE_API virtual FString ToString() const override;
 	
 	UE_API void SetSectionName(FName InSectionName);
@@ -921,29 +951,32 @@ struct FHierarchySectionViewModel : FHierarchyElementViewModel
 
 	void SetSectionImage(const FSlateBrush* InSectionImage) { SectionImage = InSectionImage; }
 	const FSlateBrush* GetSectionImage() const { return SectionImage; }
-
-	void SetDropDisallowed(bool bInDropDisallowed) { bDropDisallowed = bInDropDisallowed; }
+	
+	bool IsAllSection() const { return GetData()->GetFName() == UDataHierarchyViewModelBase::AllSectionHierarchyObjectName || GetData()->GetFName() == UDataHierarchyViewModelBase::AllSectionSourceObjectName; }
 protected:
+	virtual bool AllowEditingInDetailsPanel() const override;
+
 	/** Only hierarchy sections are editable */
-	virtual FCanPerformActionResults IsEditableByUser() override { return FCanPerformActionResults(IsForHierarchy()); }
-	virtual bool CanHaveChildren() const override { return false; }
-	UE_API virtual FCanPerformActionResults CanDragInternal() override;
+	UE_API virtual FResultWithUserFeedback IsEditableByUser() override;
+	/** Technically a section does not have 'children' but for drag & drop logic it is considered as such */
+	UE_API virtual FResultWithUserFeedback CanHaveChildren() const override { return true; }
+	UE_API virtual FResultWithUserFeedback CanDragInternal() override;
 	
 	/** We can only rename hierarchy sections */
 	UE_API virtual bool CanRenameInternal() override;
 	virtual void RenameInternal(FName NewName) override { GetDataMutable<UHierarchySection>()->SetSectionName(NewName); }
 
 	UE_API virtual bool CanDeleteInternal() override;
-	
-	UE_API virtual FCanPerformActionResults CanDropOnInternal(TSharedPtr<FHierarchyElementViewModel> DraggedItem, EItemDropZone ItemDropZone) override;
-	UE_API virtual void OnDroppedOnInternal(TSharedPtr<FHierarchyElementViewModel> DroppedItem, EItemDropZone ItemDropZone) override;
+
+	/** We override the default dropped on behavior given sections don't actually contain children. */
+	UE_API virtual bool OnDroppedOnOverride(TSharedPtr<FHierarchyElementViewModel> DroppedItem, EItemDropZone ItemDropZone) override;
+	UE_API virtual FResultWithUserFeedback CanContainInternal(TSubclassOf<UHierarchyElement> InHierarchyElement) override;
+	UE_API virtual FResultWithUserFeedback CanDropOnInternal(TSharedPtr<FHierarchyElementViewModel> DraggedItem, EItemDropZone ItemDropZone) override;
 
 	UE_API virtual void FinalizeInternal() override;
 
 private:
 	const FSlateBrush* SectionImage = nullptr;
-	
-	bool bDropDisallowed = false;
 };
 
 struct FHierarchyItemViewModel : FHierarchyElementViewModel
@@ -952,32 +985,37 @@ struct FHierarchyItemViewModel : FHierarchyElementViewModel
 	
 	virtual ~FHierarchyItemViewModel() override {}
 
-	virtual FCanPerformActionResults IsEditableByUser() override { return FCanPerformActionResults(true); }
-	virtual bool CanHaveChildren() const override { return false; }
-	virtual FCanPerformActionResults CanDragInternal() override { return true; }
+	UE_API virtual FResultWithUserFeedback IsEditableByUser() override { return FResultWithUserFeedback(true); }
+	UE_API virtual FResultWithUserFeedback CanHaveChildren() const override { return false; }
+	UE_API virtual FResultWithUserFeedback CanDragInternal() override { return true; }
+
+	UE_API virtual FResultWithUserFeedback CanContainInternal(TSubclassOf<UHierarchyElement> InHierarchyElement) override;
+	UE_API virtual FResultWithUserFeedback CanDropOnInternal(TSharedPtr<FHierarchyElementViewModel> DraggedItem, EItemDropZone ItemDropZone) override;
 
-	UE_API virtual FCanPerformActionResults CanDropOnInternal(TSharedPtr<FHierarchyElementViewModel> DraggedItem, EItemDropZone ItemDropZone) override;
-	UE_API virtual void OnDroppedOnInternal(TSharedPtr<FHierarchyElementViewModel> DroppedItem, EItemDropZone ItemDropZone) override;
+	UE_API virtual const FTableRowStyle* GetRowStyle() const override;
 };
 
 struct FHierarchyCategoryViewModel : FHierarchyElementViewModel
 {
 	FHierarchyCategoryViewModel(UHierarchyCategory* InCategory, TSharedRef<FHierarchyElementViewModel> InParent, TWeakObjectPtr<UDataHierarchyViewModelBase> InHierarchyViewModel) : FHierarchyElementViewModel(InCategory, InParent, InHierarchyViewModel, InParent->IsForHierarchy()) {}
 	virtual ~FHierarchyCategoryViewModel() override{}
-
+	
 	FText GetCategoryName() const { return GetData<UHierarchyCategory>()->GetCategoryAsText(); }
 	
-	virtual FCanPerformActionResults IsEditableByUser() override { return FCanPerformActionResults(true); }
-	virtual bool CanHaveChildren() const override { return true; }
-	virtual FCanPerformActionResults CanDragInternal() override { return true; }
-	virtual bool CanRenameInternal() override { return true; }
-	virtual void RenameInternal(FName NewName) override { GetDataMutable<UHierarchyCategory>()->SetCategoryName(NewName); }
+	UE_API virtual FResultWithUserFeedback IsEditableByUser() override { return FResultWithUserFeedback(true); }
+	UE_API virtual FResultWithUserFeedback CanHaveChildren() const override { return true; }
+	UE_API virtual FResultWithUserFeedback CanDragInternal() override { return true; }
+	UE_API virtual bool CanRenameInternal() override { return true; }
+	UE_API virtual void RenameInternal(FName NewName) override { GetDataMutable<UHierarchyCategory>()->SetCategoryName(NewName); }
 	virtual const UHierarchySection* GetSectionInternal() const override { return Cast<UHierarchyCategory>(Element)->GetSection(); }
 
 	UE_API bool IsTopCategoryActive() const;
 
-	UE_API virtual FCanPerformActionResults CanDropOnInternal(TSharedPtr<FHierarchyElementViewModel>, EItemDropZone ItemDropZone) override;
-	UE_API virtual void OnDroppedOnInternal(TSharedPtr<FHierarchyElementViewModel> DroppedItem, EItemDropZone ItemDropZone) override;
+	UE_API virtual FResultWithUserFeedback CanContainInternal(TSubclassOf<UHierarchyElement> InHierarchyElement) override;
+	UE_API virtual FResultWithUserFeedback CanDropOnInternal(TSharedPtr<FHierarchyElementViewModel>, EItemDropZone ItemDropZone) override;
+	UE_API virtual void PostAddFixup() override;
+
+	UE_API virtual const FTableRowStyle* GetRowStyle() const override;
 };
 
 class FSectionDragDropOp : public FHierarchyDragDropOp
