Make UHierarchyElement's default created identity guid deterministic based on the object path
#tests Cooked using -diffonly to confirm that the guids are deterministic and no longer generate errors for assets that haven't been saved with the hierarchy editor changes

diff --git a/Engine/Source/Editor/DataHierarchyEditor/Private/DataHierarchyViewModelBase.cpp b/Engine/Source/Editor/DataHierarchyEditor/Private/DataHierarchyViewModelBase.cpp
--- a/Engine/Source/Editor/DataHierarchyEditor/Private/DataHierarchyViewModelBase.cpp
+++ b/Engine/Source/Editor/DataHierarchyEditor/Private/DataHierarchyViewModelBase.cpp
@@ -28,6 +28,15 @@ FName UDataHierarchyViewModelBase::AllSectionHierarchyObjectName = "AllSection_H
 FName UDataHierarchyViewModelBase::AllSectionSourceObjectName = "AllSection_Source";
 FText UDataHierarchyViewModelBase::AllSectionDisplayName = LOCTEXT("AllSectionDisplayName", "All");
 
+UHierarchyElement::UHierarchyElement()
+{
+	const uint64 DeterministicSeed = 0x377206996352438CLL;
+
+	TStringBuilder<256> PathName;
+	GetPathName(nullptr, PathName);
+	Identity.Guids.Add(FGuid::NewDeterministicGuid(PathName.ToView(), DeterministicSeed));
+}
+
 const TArray<UHierarchyElement*>& UHierarchyElement::GetChildren() const
 {
 	return Children;
diff --git a/Engine/Source/Editor/DataHierarchyEditor/Public/DataHierarchyViewModelBase.h b/Engine/Source/Editor/DataHierarchyEditor/Public/DataHierarchyViewModelBase.h
--- a/Engine/Source/Editor/DataHierarchyEditor/Public/DataHierarchyViewModelBase.h
+++ b/Engine/Source/Editor/DataHierarchyEditor/Public/DataHierarchyViewModelBase.h
@@ -149,7 +149,7 @@ class UHierarchyElement : public UObject
 	GENERATED_BODY()
 
 public:
-	UHierarchyElement() { Identity.Guids.Add(FGuid::NewGuid()); }
+	UE_API UHierarchyElement();
 	virtual ~UHierarchyElement() override {}
 
 	TArray<TObjectPtr<UHierarchyElement>>& GetChildrenMutable() { return Children; }
