User Asset Tagging: v3

diff --git a/Engine/Source/Editor/DataHierarchyEditor/Private/DataHierarchyViewModelBase.cpp b/Engine/Source/Editor/DataHierarchyEditor/Private/DataHierarchyViewModelBase.cpp
--- a/Engine/Source/Editor/DataHierarchyEditor/Private/DataHierarchyViewModelBase.cpp
+++ b/Engine/Source/Editor/DataHierarchyEditor/Private/DataHierarchyViewModelBase.cpp
@@ -18,6 +18,7 @@
 #include "Framework/Notifications/NotificationManager.h"
 #include "ToolMenus.h"
 #include "Logging/StructuredLog.h"
+#include "DataHierarchyElementCustomVersion.h"
 #include "Widgets/Notifications/SNotificationList.h"
 
 #include UE_INLINE_GENERATED_CPP_BY_NAME(DataHierarchyViewModelBase)
@@ -28,6 +29,11 @@ FName UDataHierarchyViewModelBase::AllSectionHierarchyObjectName = "AllSection_H
 FName UDataHierarchyViewModelBase::AllSectionSourceObjectName = "AllSection_Source";
 FText UDataHierarchyViewModelBase::AllSectionDisplayName = LOCTEXT("AllSectionDisplayName", "All");
 
+FText FHierarchyElementViewModel::FDefaultUserFeedback::UneditableTarget = LOCTEXT("DefaultUserFeedback_UneditableTarget", "{0} is not editable.");
+FText FHierarchyElementViewModel::FDefaultUserFeedback::TargetCantContainType = LOCTEXT("DefaultUserFeedback_ViewModelCantContainChildType", "{0} can't contain elements of type {1}");
+FText FHierarchyElementViewModel::FDefaultUserFeedback::AddElementFromSourceToHierarchy = LOCTEXT("DefaultUserFeedback_AddElementFromSourceToHierarchy", "Add {0} to {1}");
+FText FHierarchyElementViewModel::FDefaultUserFeedback::MoveElementWithinHierarchy = LOCTEXT("DefaultUserFeedback_MoveElementWithinHierarchy", "Move {0} to {1}");
+
 UHierarchyElement::UHierarchyElement()
 {
 	const uint64 DeterministicSeed = 0x377206996352438CLL;
@@ -37,7 +43,7 @@ UHierarchyElement::UHierarchyElement()
 	Identity.Guids.Add(FGuid::NewDeterministicGuid(PathName.ToView(), DeterministicSeed));
 }
 
-const TArray<UHierarchyElement*>& UHierarchyElement::GetChildren() const
+TConstArrayView<const TObjectPtr<UHierarchyElement>> UHierarchyElement::GetChildren() const
 {
 	return Children;
 }
@@ -163,6 +169,13 @@ bool UHierarchyElement::Modify(bool bAlwaysMarkDirty)
 	return bSavedToTransactionBuffer;
 }
 
+void UHierarchyElement::Serialize(FArchive& Ar)
+{
+	Super::Serialize(Ar);
+	
+	Ar.UsingCustomVersion(FDataHierarchyElementCustomVersion::Guid);
+}
+
 void UHierarchyElement::PostLoad()
 {
 	if(Guid_DEPRECATED.IsValid())
@@ -189,35 +202,9 @@ void UHierarchyElement::PostLoad()
 	Super::PostLoad();
 }
 
-UHierarchySection* UHierarchyRoot::AddSection(FText InNewSectionName, int32 InsertIndex, TSubclassOf<UHierarchySection> SectionClass)
-{
-	TSet<FName> ExistingSectionNames;
-	
-	for(FName& SectionName : GetSections())
-	{
-		ExistingSectionNames.Add(SectionName);
-	}
-	
-	FName NewName = UE::DataHierarchyEditor::GetUniqueName(FName(InNewSectionName.ToString()), ExistingSectionNames);
-	UHierarchySection* NewSectionItem = NewObject<UHierarchySection>(this, SectionClass);
-	NewSectionItem->SetSectionName(NewName);
-	NewSectionItem->SetFlags(RF_Transactional);
-	
-	if(InsertIndex == INDEX_NONE)
-	{
-		Sections.Add(NewSectionItem);
-	}
-	else
-	{
-		Sections.Insert(NewSectionItem, InsertIndex);
-	}
-	
-	return NewSectionItem;
-}
-
 UHierarchySection* UHierarchyRoot::FindSectionByIdentity(FHierarchyElementIdentity SectionIdentity)
 {
-	for(UHierarchySection* Section : Sections)
+	for(UHierarchySection* Section : GetSectionDataMutable())
 	{
 		if(Section->GetPersistentIdentity() == SectionIdentity)
 		{
@@ -238,32 +225,10 @@ void UHierarchyRoot::DuplicateSectionFromOtherRoot(const UHierarchySection& Sect
 	Sections.Add(Cast<UHierarchySection>(StaticDuplicateObject(&SectionToCopy, this)));
 }
 
-void UHierarchyRoot::RemoveSection(FText SectionName)
-{
-	if(Sections.ContainsByPredicate([SectionName](UHierarchySection* Section)
-	{
-		return Section->GetSectionNameAsText().EqualTo(SectionName);
-	}))
-	{
-		Sections.RemoveAll([SectionName](UHierarchySection* Section)
-		{
-			return Section->GetSectionNameAsText().EqualTo(SectionName);
-		});
-	}
-}
-
-void UHierarchyRoot::RemoveSectionByIdentity(FHierarchyElementIdentity SectionIdentity)
-{
-	Sections.RemoveAll([SectionIdentity](UHierarchySection* Section)
-	{
-		return Section->GetPersistentIdentity() == SectionIdentity;
-	});
-}
-
-TSet<FName> UHierarchyRoot::GetSections() const
+TSet<FName> UHierarchyRoot::GetSectionNames() const
 {
 	TSet<FName> OutSections;
-	for(UHierarchySection* Section : Sections)
+	for(const UHierarchySection* Section : Sections)
 	{
 		OutSections.Add(Section->GetSectionName());
 	}
@@ -315,27 +280,6 @@ void UHierarchyRoot::Serialize(FStructuredArchive::FRecord Record)
 	Super::Serialize(Record);
 }
 
-bool FHierarchyCategoryViewModel::IsTopCategoryActive() const
-{
-	if(UHierarchyCategory* Category = GetDataMutable<UHierarchyCategory>())
-	{
-		const UHierarchyCategory* Result = Category;
-		const UHierarchyCategory* TopLevelCategory = Result;
-		
-		for (; TopLevelCategory != nullptr; TopLevelCategory = TopLevelCategory->GetTypedOuter<UHierarchyCategory>() )
-		{
-			if(TopLevelCategory != nullptr)
-			{
-				Result = TopLevelCategory;
-			}
-		}
-		
-		return HierarchyViewModel->IsHierarchySectionActive(Result->GetSection());
-	}
-
-	return false;	
-}
-
 FHierarchyElementViewModel::FResultWithUserFeedback FHierarchyCategoryViewModel::CanContainInternal(TSubclassOf<UHierarchyElement> InHierarchyElement)
 {
 	return (InHierarchyElement->IsChildOf<UHierarchyItem>() || InHierarchyElement->IsChildOf<UHierarchyCategory>());	
@@ -351,81 +295,75 @@ FHierarchyElementViewModel::FResultWithUserFeedback FHierarchyCategoryViewModel:
 	TArray<TSharedPtr<FHierarchyCategoryViewModel>> SiblingCategories;
 	Parent.Pin()->GetChildrenViewModelsForType<UHierarchyCategory, FHierarchyCategoryViewModel>(SiblingCategories);
 	
-	// we only allow drops if some general conditions are fulfilled
-	if(DraggedElement->GetData() != GetData() &&
-		(!DraggedElement->HasParent(AsShared(), false) || ItemDropZone != EItemDropZone::OntoItem)  &&
-		!HasParent(DraggedElement, true))
+	// categories can be dropped on categories, but only if the resulting sibling categories or children categories have different names
+	if(DraggedElement->GetData()->IsA<UHierarchyCategory>())
 	{
-		// categories can be dropped on categories, but only if the resulting sibling categories or children categories have different names
-		if(DraggedElement->GetData()->IsA<UHierarchyCategory>())
+		if(ItemDropZone != EItemDropZone::OntoItem)
 		{
-			if(ItemDropZone != EItemDropZone::OntoItem)
-			{
-				bool bContainsSiblingWithSameName = SiblingCategories.ContainsByPredicate([DraggedElement](TSharedPtr<FHierarchyCategoryViewModel> HierarchyCategoryViewModel)
-					{
-						return DraggedElement->ToString() == HierarchyCategoryViewModel->ToString() && DraggedElement != HierarchyCategoryViewModel;
-					});
-
-				if(bContainsSiblingWithSameName)
+			bool bContainsSiblingWithSameName = SiblingCategories.ContainsByPredicate([DraggedElement](TSharedPtr<FHierarchyCategoryViewModel> HierarchyCategoryViewModel)
 				{
-					Results.bResult = false;
-					Results.UserFeedback = LOCTEXT("CantDropCategorNextToCategorySameSiblingNames", "A category of the same name already exists here, potentially in a different section. Please rename your category first.");
-					return Results;
-				}
+					return DraggedElement->ToString() == HierarchyCategoryViewModel->ToString() && DraggedElement != HierarchyCategoryViewModel;
+				});
 
-				Results.UserFeedback = LOCTEXT("MoveCategoryText", "Move category here");
+			if(bContainsSiblingWithSameName)
+			{
+				Results.bResult = false;
+				Results.UserFeedback = LOCTEXT("CantDropCategorNextToCategorySameSiblingNames", "A category of the same name already exists here, potentially in a different section. Please rename your category first.");
+				return Results;
+			}
 
-				// if we are making a category a sibling of another at the root level, the section will be set to the currently active section. Let that be known.
-				if(Parent.Pin()->GetData()->IsA<UHierarchyRoot>())
+			Results.UserFeedback = LOCTEXT("MoveCategoryText", "Move category here");
+
+			// if we are making a category a sibling of another at the root level, the section will be set to the currently active section. Let that be known.
+			if(Parent.Pin()->GetData()->IsA<UHierarchyRoot>())
+			{
+				UHierarchyCategory* DraggedCategory = Cast<UHierarchyCategory>(DraggedElement->GetDataMutable());
+				
+				FDataHierarchyElementMetaData_SectionAssociation SectionAssociation = DraggedCategory->FindMetaDataOfTypeOrDefault<FDataHierarchyElementMetaData_SectionAssociation>();					
+				if(SectionAssociation.Section != HierarchyViewModel->GetActiveHierarchySectionData())
 				{
-					UHierarchyCategory* DraggedCategory = Cast<UHierarchyCategory>(DraggedElement->GetDataMutable());
-					const UHierarchySection* CurrentElementSection = DraggedCategory->GetSection() == nullptr ? HierarchyViewModel->GetDefaultHierarchySectionViewModel()->GetData<UHierarchySection>() : DraggedCategory->GetSection();
-					
-					if(CurrentElementSection != HierarchyViewModel->GetActiveHierarchySectionData())
-					{
-						FText SectionChangeBaseText = LOCTEXT("CategorySectionWillUpdateDueToDrop", "The section of the category will change to {0} after the drop");
-						FText ActualSectionChangeText = FText::FormatOrdered(SectionChangeBaseText, HierarchyViewModel->GetActiveHierarchySectionViewModel()->GetSectionNameAsText());
-						Results.UserFeedback = FText::FormatOrdered(FText::AsCultureInvariant("{0}\n{1}"), Results.UserFeedback.Get(FText::GetEmpty()), ActualSectionChangeText);
-					}
+					FText SectionChangeBaseText = LOCTEXT("CategorySectionWillUpdateDueToDrop", "The section of the category will change to {0} after the drop");
+					FText ActualSectionChangeText = FText::FormatOrdered(SectionChangeBaseText, HierarchyViewModel->GetActiveHierarchySectionViewModel()->GetSectionNameAsText());
+					Results.UserFeedback = FText::FormatOrdered(FText::AsCultureInvariant("{0}\n{1}"), Results.UserFeedback.Get(FText::GetEmpty()), ActualSectionChangeText);
 				}
 			}
-			else
-			{
-				bool bContainsChildrenCategoriesWithSameName = TargetChildrenCategories.ContainsByPredicate([DraggedElement](TSharedPtr<FHierarchyCategoryViewModel> HierarchyCategoryViewModel)
-					{
-						return DraggedElement->ToString() == HierarchyCategoryViewModel->ToString();
-					});
-
-				if(bContainsChildrenCategoriesWithSameName)
+		}
+		else
+		{
+			bool bContainsChildrenCategoriesWithSameName = TargetChildrenCategories.ContainsByPredicate([DraggedElement](TSharedPtr<FHierarchyCategoryViewModel> HierarchyCategoryViewModel)
 				{
-					Results.bResult = false;
-					Results.UserFeedback = LOCTEXT("CantDropCategoryOnCategorySameChildCategoryName", "A sub-category of the same name already exists! Please rename your category first.");
-					return Results;
-				}
+					return DraggedElement->ToString() == HierarchyCategoryViewModel->ToString();
+				});
 
-				Results.UserFeedback = LOCTEXT("CreateSubcategory", "Drop category here to create a sub-category");
+			if(bContainsChildrenCategoriesWithSameName)
+			{
+				Results.bResult = false;
+				Results.UserFeedback = LOCTEXT("CantDropCategoryOnCategorySameChildCategoryName", "A sub-category of the same name already exists! Please rename your category first.");
+				return Results;
 			}
 
-			Results.bResult = true;
-			return Results;
+			Results.UserFeedback = LOCTEXT("CreateSubcategory", "Drop category here to create a sub-category");
 		}
-		else if(DraggedElement->GetData()->IsA<UHierarchyItem>())
-		{
-			// items can generally be dropped onto categories
-			Results.bResult = EItemDropZone::OntoItem == ItemDropZone;
 
-			if(Results.bResult)
+		Results.bResult = true;
+		return Results;
+	}
+	else if(DraggedElement->GetData()->IsA<UHierarchyItem>())
+	{
+		// items can generally be dropped onto categories
+		Results.bResult = EItemDropZone::OntoItem == ItemDropZone;
+
+		if(Results.bResult)
+		{
+			if(DraggedElement->IsForHierarchy() == false)
 			{
-				if(DraggedElement->IsForHierarchy() == false)
-				{
-					FText Message = LOCTEXT("AddItemToCategoryDragMessage", "Add {0} to {1}");
-					Results.UserFeedback = FText::FormatOrdered(Message, FText::FromString(DraggedElement->ToString()), FText::FromString(ToString()));
-				}
-				else
-				{
-					FText Message = LOCTEXT("MoveItemToCategoryDragMessage", "Move {0} to {1}");
-					Results.UserFeedback = FText::FormatOrdered(Message, FText::FromString(DraggedElement->ToString()), FText::FromString(ToString()));
-				}
+				FText Message = LOCTEXT("AddItemToCategoryDragMessage", "Add {0} to {1}");
+				Results.UserFeedback = FText::FormatOrdered(Message, FText::FromString(DraggedElement->ToString()), FText::FromString(ToString()));
+			}
+			else
+			{
+				FText Message = LOCTEXT("MoveItemToCategoryDragMessage", "Move {0} to {1}");
+				Results.UserFeedback = FText::FormatOrdered(Message, FText::FromString(DraggedElement->ToString()), FText::FromString(ToString()));
 			}
 		}
 	}
@@ -433,18 +371,6 @@ FHierarchyElementViewModel::FResultWithUserFeedback FHierarchyCategoryViewModel:
 	return Results;
 }
 
-void UHierarchyCategory::FixupSectionLinkage()
-{
-	UHierarchyRoot* OwningRoot = GetTypedOuter<UHierarchyRoot>();
-
-	if(Section != nullptr && Section->GetTypedOuter<UHierarchyRoot>() != OwningRoot)
-	{
-		UHierarchySection* CorrectSection = OwningRoot->FindSectionByIdentity(Section->GetPersistentIdentity());
-		ensure(CorrectSection != nullptr);
-		Section = CorrectSection;
-	}
-}
-
 FHierarchyElementIdentity UHierarchyCategory::ConstructIdentity()
 {
 	FHierarchyElementIdentity Identity;
@@ -462,6 +388,18 @@ void UHierarchyCategory::PostLoad()
 	{
 		SetIdentity(UHierarchyCategory::ConstructIdentity());
 	}
+	
+	const int32 CustomVer = GetLinkerCustomVersion(FDataHierarchyElementCustomVersion::Guid);
+	
+	if(CustomVer < FDataHierarchyElementCustomVersion::Type::MetaDataIntroduction)
+	{
+		if(Section_DEPRECATED)
+		{
+			FDataHierarchyElementMetaData_SectionAssociation* SectionAssociation = FindOrAddMetaDataOfType<FDataHierarchyElementMetaData_SectionAssociation>();
+			SectionAssociation->Section = Section_DEPRECATED;
+			Section_DEPRECATED = nullptr;
+		}
+	}
 }
 
 void UHierarchySection::SetSectionNameAsText(const FText& Text)
@@ -518,7 +456,6 @@ void UDataHierarchyViewModelBase::Initialize()
 	
 	HierarchyRootViewModel->Initialize();
 	HierarchyRootViewModel->AddChildFilter(FHierarchyElementViewModel::FOnFilterChild::CreateUObject(this, &UDataHierarchyViewModelBase::FilterForHierarchySection));
-	HierarchyRootViewModel->AddChildFilter(FHierarchyElementViewModel::FOnFilterChild::CreateUObject(this, &UDataHierarchyViewModelBase::FilterForUncategorizedRootItemsInAllSection));
 	HierarchyRootViewModel->SyncViewModelsToData();
 
 	AllSection = NewObject<UHierarchySection>(this, AllSectionHierarchyObjectName, RF_Transient);
@@ -668,23 +605,8 @@ bool UDataHierarchyViewModelBase::FilterForHierarchySection(TSharedPtr<const FHi
 		}
 
 		// if not, we check against identical section data
-		return ActiveHierarchySectionViewModel.Pin()->GetData() == ViewModel->GetSection();
-	}
-
-	return true;
-}
-
-bool UDataHierarchyViewModelBase::FilterForUncategorizedRootItemsInAllSection(TSharedPtr<const FHierarchyElementViewModel> ViewModel) const
-{
-	if(ActiveHierarchySectionViewModel.IsValid())
-	{
-		// we want to filter out all items that are directly added to the root if we aren't in the 'All' section
-		if(ActiveHierarchySectionViewModel.Pin()->IsAllSection())
-		{
-			return true;
-		}
-		
-		return ViewModel->GetData<UHierarchyCategory>() != nullptr;
+		FDataHierarchyElementMetaData_SectionAssociation SectionAssociation = ViewModel->GetData()->FindMetaDataOfTypeOrDefault<FDataHierarchyElementMetaData_SectionAssociation>();		
+		return ActiveHierarchySectionViewModel.Pin()->GetSectionData() == SectionAssociation.Section;
 	}
 
 	return true;
@@ -801,14 +723,14 @@ TSharedPtr<FHierarchySectionViewModel> UDataHierarchyViewModelBase::GetActiveHie
 	return ActiveHierarchySectionViewModel.Pin();
 }
 
-UHierarchySection* UDataHierarchyViewModelBase::GetActiveHierarchySectionData() const
+const UHierarchySection* UDataHierarchyViewModelBase::GetActiveHierarchySectionData() const
 {
-	return ActiveHierarchySectionViewModel.IsValid() ? ActiveHierarchySectionViewModel.Pin()->GetDataMutable<UHierarchySection>() : nullptr;
+	return ActiveHierarchySectionViewModel.IsValid() ? ActiveHierarchySectionViewModel.Pin()->GetSectionData() : nullptr;
 }
 
 bool UDataHierarchyViewModelBase::IsHierarchySectionActive(const UHierarchySection* Section) const
 {
-	return ActiveHierarchySectionViewModel.IsValid() ? ActiveHierarchySectionViewModel.Pin()->GetData() == Section : false;
+	return ActiveHierarchySectionViewModel.IsValid() ? ActiveHierarchySectionViewModel.Pin()->GetSectionData() == Section : false;
 }
 
 FString UDataHierarchyViewModelBase::OnElementToStringDebug(TSharedPtr<FHierarchyElementViewModel> ElementViewModel) const
@@ -822,17 +744,24 @@ FHierarchyElementViewModel::~FHierarchyElementViewModel()
 	FilteredChildren.Empty();
 }
 
-UHierarchyElement* FHierarchyElementViewModel::AddChild(TSubclassOf<UHierarchyElement> NewChildClass, FHierarchyElementIdentity ChildIdentity)
+TSharedPtr<FHierarchyElementViewModel> FHierarchyElementViewModel::AddChild(TSubclassOf<UHierarchyElement> NewChildClass, int32 InsertIndex, FHierarchyElementIdentity ChildIdentity)
 {
-	UHierarchyElement* NewChild = NewObject<UHierarchyElement>(GetDataMutable(), NewChildClass);
+	UHierarchyElement* NewChild = GetDataMutable()->AddChild<UHierarchyElement>(NewChildClass, InsertIndex);
 	NewChild->SetFlags(RF_Transactional);
-	NewChild->Modify();
 	NewChild->SetIdentity(ChildIdentity);
-	GetDataMutable()->GetChildrenMutable().Add(NewChild);
+	NewChild->Modify();
 	
 	SyncViewModelsToData();
+
+	TSharedPtr<FHierarchyElementViewModel> ElementViewModel = FindViewModelForChild(NewChild);
+	ensureMsgf(ElementViewModel.IsValid(), TEXT("Invalid view model detected after syncing view models to data."));
+	
+	TInstancedStruct<FHierarchyElementChangedPayload_AddedElement> Payload = TInstancedStruct<FHierarchyElementChangedPayload_AddedElement>::Make();
+	Payload.GetMutable<FHierarchyElementChangedPayload_AddedElement>().AddedElementViewModel = ElementViewModel;
+	HierarchyViewModel->OnHierarchyElementChanged().Execute(Payload);
+	
 	HierarchyViewModel->OnHierarchyChanged().Broadcast();
-	return NewChild;
+	return ElementViewModel;
 }
 
 void FHierarchyElementViewModel::Tick(float DeltaTime)
@@ -957,12 +886,12 @@ const TArray<TSharedPtr<FHierarchyElementViewModel>>& FHierarchyElementViewModel
 	return FilteredChildren;
 }
 
-void FHierarchyElementViewModel::SortChildrenData() const
+void FHierarchyElementViewModel::SortChildrenData()
 {
 	GetDataMutable()->GetChildrenMutable().StableSort([](const UHierarchyElement& ItemA, const UHierarchyElement& ItemB)
-		{
-			return ItemA.IsA<UHierarchyCategory>() && ItemB.IsA<UHierarchyItem>();
-		});
+	{
+		return ItemA.IsA<UHierarchyCategory>() && ItemB.IsA<UHierarchyItem>();
+	});
 }
 
 int32 FHierarchyElementViewModel::GetHierarchyDepth() const
@@ -1038,7 +967,7 @@ TSharedRef<FHierarchyElementViewModel> FHierarchyElementViewModel::ReparentToThi
 	return ViewModel.ToSharedRef();
 }
 
-TSharedPtr<FHierarchyElementViewModel> FHierarchyElementViewModel::FindViewModelForChild(UHierarchyElement* Child, bool bSearchRecursively) const
+TSharedPtr<FHierarchyElementViewModel> FHierarchyElementViewModel::FindViewModelForChild(const UHierarchyElement* Child, bool bSearchRecursively) const
 {
 	int32 Index = FindIndexOfChild(Child);
 	if(Index != INDEX_NONE)
@@ -1088,7 +1017,7 @@ TSharedPtr<FHierarchyElementViewModel> FHierarchyElementViewModel::FindViewModel
 	return nullptr;
 }
 
-int32 FHierarchyElementViewModel::FindIndexOfChild(UHierarchyElement* Child) const
+int32 FHierarchyElementViewModel::FindIndexOfChild(const UHierarchyElement* Child) const
 {
 	return Children.FindLastByPredicate([Child](TSharedPtr<FHierarchyElementViewModel> Item)
 	{
@@ -1096,14 +1025,14 @@ int32 FHierarchyElementViewModel::FindIndexOfChild(UHierarchyElement* Child) con
 	});
 }
 
-int32 FHierarchyElementViewModel::FindIndexOfDataChild(TSharedPtr<FHierarchyElementViewModel> Child) const
+int32 FHierarchyElementViewModel::FindIndexOfDataChild(TSharedPtr<const FHierarchyElementViewModel> Child) const
 {
-	return GetData()->GetChildren().Find(Child->GetDataMutable());
+	return GetData()->GetChildren().IndexOfByKey(Child->GetData());
 }
 
-int32 FHierarchyElementViewModel::FindIndexOfDataChild(UHierarchyElement* Child) const
+int32 FHierarchyElementViewModel::FindIndexOfDataChild(const UHierarchyElement* Child) const
 {
-	return GetData()->GetChildren().Find(Child);
+	return GetData()->GetChildren().IndexOfByKey(Child);
 }
 
 void FHierarchyElementViewModel::Delete()
@@ -1117,6 +1046,10 @@ void FHierarchyElementViewModel::DeleteChild(TSharedPtr<FHierarchyElementViewMod
 	GetDataMutable()->Modify();
 	GetDataMutable()->GetChildrenMutable().Remove(Child->GetDataMutable());
 	Children.Remove(Child);
+
+	TInstancedStruct<FHierarchyElementChangedPayload_DeletedElement> Payload = TInstancedStruct<FHierarchyElementChangedPayload_DeletedElement>::Make();
+	Payload.GetMutable<FHierarchyElementChangedPayload_DeletedElement>().DeletedElementViewModel = Child;
+	HierarchyViewModel->OnHierarchyElementChanged().Execute(Payload);
 }
 
 TOptional<EItemDropZone> FHierarchyElementViewModel::OnCanRowAcceptDrop(const FDragDropEvent& DragDropEvent, EItemDropZone ItemDropZone, TSharedPtr<FHierarchyElementViewModel> Item)
@@ -1164,29 +1097,58 @@ FHierarchyElementViewModel::FResultWithUserFeedback FHierarchyElementViewModel::
 	return CanDragInternal();
 }
 
-FHierarchyElementViewModel::FResultWithUserFeedback FHierarchyElementViewModel::CanDropOn(TSharedPtr<FHierarchyElementViewModel> DraggedItem, EItemDropZone ItemDropZone)
+bool FHierarchyElementViewModel::IsExpandedByDefault() const
+{
+	// Otherwise, let the element decide
+	return IsExpandedByDefaultInternal();
+}
+
+FHierarchyElementViewModel::FResultWithUserFeedback FHierarchyElementViewModel::CanDropOn(TSharedPtr<FHierarchyElementViewModel> DraggedElement, EItemDropZone ItemDropZone)
 {
 	// Editability check
 	{
 		FResultWithUserFeedback Results = IsEditableByUser();
 		if(Results.bResult == false)
 		{
+			if(Results.UserFeedback.IsSet() == false)
+			{
+				Results.UserFeedback = FText::FormatOrdered(FDefaultUserFeedback::UneditableTarget, this->ToStringAsText());
+			}
+			
 			return Results;
 		}
 	}
 
 	// Self identity check
-	if(DraggedItem->GetData() == GetData())
+	if(DraggedElement->GetData() == GetData())
 	{
 		return false;
 	}
 
+	// Immediate parent check; refuse if we are already parented to the target
+	if(DraggedElement->HasParent(AsShared(), false) && ItemDropZone == EItemDropZone::OntoItem)
+	{
+		return false;
+	}
+
+	// Recursive child check; refuse if we try to drag a parent onto its child
+	if(HasParent(DraggedElement, true))
+	{
+		return false;
+	}
+	
 	// Can Contain checks
 	if(ItemDropZone == EItemDropZone::OntoItem)
 	{
-		FResultWithUserFeedback Results = CanContain(DraggedItem->GetData()->GetClass());
+		FResultWithUserFeedback Results = CanContain(DraggedElement->GetData()->GetClass());
 		if(Results.bResult == false)
 		{
+			if(Results.UserFeedback.IsSet() == false)
+			{
+				Results.UserFeedback = FText::FormatOrdered(FDefaultUserFeedback::TargetCantContainType,
+					GetData()->GetClass()->GetDisplayNameText(), DraggedElement->GetData()->GetClass()->GetDisplayNameText());
+			}
+			
 			return Results;
 		}
 	}
@@ -1194,16 +1156,36 @@ FHierarchyElementViewModel::FResultWithUserFeedback FHierarchyElementViewModel::
 	{
 		if(Parent.IsValid())
 		{
-			FResultWithUserFeedback Results = Parent.Pin()->CanContain(DraggedItem->GetData()->GetClass());
+			FResultWithUserFeedback Results = Parent.Pin()->CanContain(DraggedElement->GetData()->GetClass());
 			if(Results.bResult == false)
 			{
+				if(Results.UserFeedback.IsSet() == false)
+				{
+					Results.UserFeedback = FText::FormatOrdered(FDefaultUserFeedback::TargetCantContainType,
+						GetData()->GetClass()->GetDisplayNameText(), Parent.Pin()->GetData()->GetClass()->GetDisplayNameText());
+				}
+				
 				return Results;
 			}
 		}
 	}
 
 	// Internal check
-	return CanDropOnInternal(DraggedItem, ItemDropZone);
+	FResultWithUserFeedback Results = CanDropOnInternal(DraggedElement, ItemDropZone);
+	
+	if(Results.bResult && Results.UserFeedback.IsSet() == false)
+	{
+		if(DraggedElement->IsForHierarchy() == false)
+		{
+			Results.UserFeedback = FText::FormatOrdered(FDefaultUserFeedback::AddElementFromSourceToHierarchy, DraggedElement->ToStringAsText(), this->ToStringAsText());
+		}
+		else
+		{
+			Results.UserFeedback = FText::FormatOrdered(FDefaultUserFeedback::MoveElementWithinHierarchy, DraggedElement->ToStringAsText(), this->ToStringAsText());
+		}
+	}
+	
+	return Results;;
 }
 
 void FHierarchyElementViewModel::OnDroppedOn(TSharedPtr<FHierarchyElementViewModel> DroppedItem, EItemDropZone ItemDropZone)
@@ -1215,56 +1197,57 @@ void FHierarchyElementViewModel::OnDroppedOn(TSharedPtr<FHierarchyElementViewMod
 	
 	FScopedTransaction Transaction(LOCTEXT("Transaction_MovedItem", "Moved an item in the hierarchy"));
 	HierarchyViewModel->GetHierarchyRoot()->Modify();
-	
-	bool bDropSucceeded = false;
-	TSharedPtr<FHierarchyElementViewModel> NewItemViewModel;
+
+	// We calculate the target view model (this if OntoItem, Parent if Above/Below the current item)
+	// and the insert index (INDEX_NONE means adding at the end)
+	TSharedPtr<FHierarchyElementViewModel> TargetViewModel;
+	int32 InsertIndex = INDEX_NONE;
 	if(ItemDropZone == EItemDropZone::OntoItem)
 	{
-		if(DroppedItem->IsForHierarchy() == false)
-		{
-			NewItemViewModel = DuplicateToThis(DroppedItem);
-		}
-		else
-		{
-			NewItemViewModel = ReparentToThis(DroppedItem);
-		}
-
-		bDropSucceeded = true;
+		TargetViewModel = AsShared();
+		InsertIndex = INDEX_NONE;
 	}
-	else if(ItemDropZone == EItemDropZone::AboveItem)
+	else
 	{
-		int32 IndexOfThis = Parent.Pin()->FindIndexOfDataChild(AsShared());
-
-		if(DroppedItem->IsForHierarchy() == false)
-		{
-			NewItemViewModel = Parent.Pin()->DuplicateToThis(DroppedItem, FMath::Max(IndexOfThis, 0));
-		}
-		else
+		if(Parent.IsValid())
 		{
-			NewItemViewModel = Parent.Pin()->ReparentToThis(DroppedItem, FMath::Max(IndexOfThis, 0));
-		}
+			TargetViewModel =  Parent.Pin();
+			int32 IndexOfThis = Parent.Pin()->FindIndexOfDataChild(AsShared());
 
-		bDropSucceeded = true;
+			// If above an item, we insert at this item's index
+			if(ItemDropZone == EItemDropZone::AboveItem)
+			{
+				InsertIndex = FMath::Max(IndexOfThis, 0);
+			}
+			// If below an item, we insert one index behind this item's index
+			else
+			{
+				InsertIndex = FMath::Min(IndexOfThis+1, Parent.Pin()->GetChildren().Num());
+			}
+		}
 	}
-	else if(ItemDropZone == EItemDropZone::BelowItem)
+
+	TSharedPtr<FHierarchyElementViewModel> NewItemViewModel;
+	if(ensure(TargetViewModel.IsValid()))
 	{
-		int32 IndexOfThis = Parent.Pin()->FindIndexOfDataChild(AsShared());
-		
+		// If the dropped item is from the source, we duplicate it over
 		if(DroppedItem->IsForHierarchy() == false)
 		{
-			NewItemViewModel = Parent.Pin()->DuplicateToThis(DroppedItem, FMath::Min(IndexOfThis+1, Parent.Pin()->GetChildren().Num()));
+			NewItemViewModel = TargetViewModel->DuplicateToThis(DroppedItem, InsertIndex);
 		}
+		// If the dropped item is from the hierarchy, we simply move it over
 		else
 		{
-			NewItemViewModel = Parent.Pin()->ReparentToThis(DroppedItem, FMath::Min(IndexOfThis+1, Parent.Pin()->GetChildren().Num()));
+			NewItemViewModel = TargetViewModel->ReparentToThis(DroppedItem, InsertIndex);
 		}
-
-		bDropSucceeded = true;
 	}
 	
-	if(bDropSucceeded)
+	if(NewItemViewModel.IsValid())
 	{
 		NewItemViewModel->PostAddFixup();
+
+		TargetViewModel->PostOnDroppedOn(NewItemViewModel);
+		
 		HierarchyViewModel->RefreshHierarchyView();
 		HierarchyViewModel->RefreshSourceView();
 	}
@@ -1298,7 +1281,13 @@ FHierarchyElementViewModel::FResultWithUserFeedback FHierarchyElementViewModel::
 	return CanContainInternal(HierarchyElementType);
 }
 
-FHierarchyElementViewModel::FResultWithUserFeedback FHierarchyElementViewModel::CanDropOnInternal(TSharedPtr<FHierarchyElementViewModel>, EItemDropZone ItemDropZone)
+bool FHierarchyElementViewModel::IsExpandedByDefaultInternal() const
+{
+	// By default, we expand all high-level elements
+	return GetHierarchyDepth() <= 1;
+}
+
+FHierarchyElementViewModel::FResultWithUserFeedback FHierarchyElementViewModel::CanDropOnInternal(TSharedPtr<FHierarchyElementViewModel> DraggedElement, EItemDropZone ItemDropZone)
 {
 	return true;
 }
@@ -1367,48 +1356,23 @@ void FHierarchyRootViewModel::Initialize()
 	GetOnSynced().BindSP(this, &FHierarchyRootViewModel::PropagateOnSynced);
 }
 
-FHierarchyElementViewModel::FResultWithUserFeedback FHierarchyRootViewModel::CanContainInternal(TSubclassOf<UHierarchyElement> InHierarchyElement)
+FHierarchyElementViewModel::FResultWithUserFeedback FHierarchyRootViewModel::CanContainInternal(TSubclassOf<UHierarchyElement> InHierarchyElementType)
 {
-	return true;
+	// The root generally accepts items and categories, not sections
+	FResultWithUserFeedback Results(InHierarchyElementType->IsChildOf<UHierarchyItem>() || InHierarchyElementType->IsChildOf<UHierarchyCategory>());
+	return Results;
 }
 
 FHierarchyElementViewModel::FResultWithUserFeedback FHierarchyRootViewModel::CanDropOnInternal(TSharedPtr<FHierarchyElementViewModel> DraggedElement, EItemDropZone ItemDropZone)
 {
-	FResultWithUserFeedback Results(false);
-	
-	// we only allow drops if some general conditions are fulfilled
-	if(DraggedElement->GetData() != GetData() &&
-		(!DraggedElement->HasParent(AsShared(), false) || ItemDropZone != EItemDropZone::OntoItem)  &&
-		!HasParent(DraggedElement, true))
-	{
-		Results.bResult = 
-			// items can be dropped onto the root directly if the section is set to "All"
-			(DraggedElement->GetData()->IsA<UHierarchyItem>() && HierarchyViewModel->GetActiveHierarchySectionViewModel()->IsAllSection())
-				||
-			// categories can be dropped onto the root always
-			(DraggedElement->GetData()->IsA<UHierarchyCategory>());
+	FResultWithUserFeedback Results(true);
 
-		if(Results.bResult)
-		{
-			if(DraggedElement->IsForHierarchy() == false)
-			{
-				FText Message = LOCTEXT("CanDropSourceItemOnRootDragMessage", "Add {0} to the hierarchy root.");
-				Message = FText::FormatOrdered(Message, FText::FromString(DraggedElement->ToString()));
-				Results.UserFeedback = Message;
-			}
-			else
-			{
-				FText Message = LOCTEXT("CanDropHierarchyItemOnRootDragMessage", "Move {0} to the hierarchy root.");
-				Message = FText::FormatOrdered(Message, FText::FromString(DraggedElement->ToString()));
-				Results.UserFeedback = Message;
-			}			
-		}
-		else
-		{
-			FText Message = LOCTEXT("CantDropHierarchyItemOnRootDragMessage", "Can not add {0} here. Please add it to a category!");
-			Message = FText::FormatOrdered(Message, FText::FromString(DraggedElement->ToString()));
-			Results.UserFeedback = Message;
-		}
+	// If we are not in the All section, we only accept categories to drop directly on the root
+	if(HierarchyViewModel->GetActiveHierarchySectionViewModel()->IsAllSection() == false && !DraggedElement->GetData()->IsA<UHierarchyCategory>())
+	{
+		Results.bResult = false;
+		FText BaseMessage = LOCTEXT("RootCanOnlyContainCategories", "While section {0} is active, the root only accepts categories. Switch to the All section first.");
+		Results.UserFeedback = FText::FormatOrdered(BaseMessage, HierarchyViewModel->GetActiveHierarchySectionViewModel()->GetSectionNameAsText());
 	}
 	
 	return Results;
@@ -1416,24 +1380,38 @@ FHierarchyElementViewModel::FResultWithUserFeedback FHierarchyRootViewModel::Can
 
 TSharedPtr<FHierarchySectionViewModel> FHierarchyRootViewModel::AddSection()
 {
-	FScopedTransaction ScopedTransaction(LOCTEXT("NewSectionAdded","Added Section"));
+	FScopedTransaction ScopedTransaction(LOCTEXT("NewSectionAdded", "Added Section"));
 	HierarchyViewModel->GetHierarchyRoot()->Modify();
 	
-	UHierarchySection* SectionData = GetDataMutable<UHierarchyRoot>()->AddSection(LOCTEXT("HierarchyEditorDefaultNewSectionName", "Section"), 0, HierarchyViewModel->GetSectionDataClass());
-	SectionData->Modify();
-	TSharedPtr<FHierarchyElementViewModel> ViewModel = HierarchyViewModel->CreateViewModelForElement(SectionData, StaticCastSharedRef<FHierarchyRootViewModel>(AsShared()));
-	TSharedPtr<FHierarchySectionViewModel> SectionViewModel = StaticCastSharedPtr<FHierarchySectionViewModel>(ViewModel);
-	if(!ensureMsgf(SectionViewModel.IsValid(), TEXT("Make sure that CreateViewModelForData creates a FHierarchySectionViewModel (or derived) for UHierarchySection elements")))
+	TSet<FName> ExistingSectionNames;
+    	
+	for(FName& SectionName : GetData<UHierarchyRoot>()->GetSectionNames())
 	{
-		return nullptr;	
+		ExistingSectionNames.Add(SectionName);
 	}
-	
-	SectionViewModels.Add(SectionViewModel);
+
+	FName NewName = UE::DataHierarchyEditor::GetUniqueName(FName(LOCTEXT("HierarchyEditorDefaultNewSectionName", "Section").ToString()), ExistingSectionNames);
+
+	UHierarchySection* NewSection = NewObject<UHierarchySection>(GetDataMutable(), HierarchyViewModel->GetSectionDataClass());
+	GetDataMutable<UHierarchyRoot>()->GetSectionDataMutable().Insert(NewSection, 0);
 	SyncViewModelsToData();
+	
+	TSharedPtr<FHierarchySectionViewModel> SectionViewModel = SectionViewModels[0];
+	ensureMsgf(SectionViewModel->GetData() == NewSection, TEXT("The section view model for a new section should always be at index 0 after syncing view models."));
+	
+	SectionViewModel->SetSectionName(NewName);
+	SectionViewModel->GetDataMutable()->Modify();
+	
 	HierarchyViewModel->SetActiveHierarchySection(SectionViewModel);
-
+	
 	OnSectionAddedDelegate.ExecuteIfBound(SectionViewModel);
+	
+	TInstancedStruct<FHierarchyElementChangedPayload_AddedElement> Payload = TInstancedStruct<FHierarchyElementChangedPayload_AddedElement>::Make();
+	Payload.GetMutable<FHierarchyElementChangedPayload_AddedElement>().AddedElementViewModel = SectionViewModel;
+	HierarchyViewModel->OnHierarchyElementChanged().Execute(Payload);
+	
 	OnSectionsChangedDelegate.ExecuteIfBound();
+	
 	return SectionViewModel;
 }
 
@@ -1444,14 +1422,14 @@ void FHierarchyRootViewModel::DeleteSection(TSharedPtr<FHierarchyElementViewMode
 	SectionViewModels.Remove(SectionViewModel);
 
 	OnSectionDeletedDelegate.ExecuteIfBound(SectionViewModel);
+	
+	TInstancedStruct<FHierarchyElementChangedPayload_DeletedElement> Payload = TInstancedStruct<FHierarchyElementChangedPayload_DeletedElement>::Make();
+	Payload.GetMutable<FHierarchyElementChangedPayload_DeletedElement>().DeletedElementViewModel = SectionViewModel;
+	HierarchyViewModel->OnHierarchyElementChanged().Execute(Payload);
+	
 	OnSectionsChangedDelegate.ExecuteIfBound();
 }
 
-void FHierarchyRootViewModel::PropagateOnSynced()
-{
-	OnSyncPropagatedDelegate.ExecuteIfBound();
-}
-
 void FHierarchyRootViewModel::SyncViewModelsToDataInternal()
 {
 	const UHierarchyRoot* RootData = GetData<UHierarchyRoot>();
@@ -1488,8 +1466,8 @@ void FHierarchyRootViewModel::SyncViewModelsToDataInternal()
 
 		if(SectionViewModel == nullptr)
 		{
-			SectionViewModel = MakeShared<FHierarchySectionViewModel>(Section, StaticCastSharedRef<FHierarchyRootViewModel>(AsShared()), HierarchyViewModel);
-			SectionViewModel->SyncViewModelsToData();;
+			SectionViewModel = StaticCastSharedPtr<FHierarchySectionViewModel>(HierarchyViewModel->CreateViewModelForElement(Section, StaticCastSharedRef<FHierarchyRootViewModel>(AsShared())));
+			SectionViewModel->SyncViewModelsToData();
 		}
 		
 		NewSectionViewModels.Add(SectionViewModel);
@@ -1512,19 +1490,39 @@ void FHierarchyRootViewModel::SyncViewModelsToDataInternal()
 		});
 }
 
+void FHierarchyRootViewModel::PostOnDroppedOn(TSharedPtr<FHierarchyElementViewModel> DroppedElementViewModel)
+{
+	// If we are in the all section, we clear the section association
+	if(HierarchyViewModel->GetActiveHierarchySectionViewModel()->IsAllSection())
+	{
+		DroppedElementViewModel->GetDataMutable()->DeleteMetaDataOfType<FDataHierarchyElementMetaData_SectionAssociation>();
+	}
+	// If not, we update the section association
+	else
+	{
+		FDataHierarchyElementMetaData_SectionAssociation* SectionAssociation = DroppedElementViewModel->GetDataMutable()->FindOrAddMetaDataOfType<FDataHierarchyElementMetaData_SectionAssociation>();
+		SectionAssociation->Section = HierarchyViewModel->GetActiveHierarchySectionViewModel()->GetSectionData();
+	}
+}
+
+void FHierarchyRootViewModel::PropagateOnSynced()
+{
+	OnSyncPropagatedDelegate.ExecuteIfBound();
+}
+
 FHierarchySectionViewModel::FHierarchySectionViewModel(UHierarchySection* InItem, TSharedRef<FHierarchyRootViewModel> InParent, TWeakObjectPtr<UDataHierarchyViewModelBase> InHierarchyViewModel)
 	: FHierarchyElementViewModel(InItem, InParent, InHierarchyViewModel, InParent->IsForHierarchy())
 {
 }
 
-UHierarchySection* FHierarchySectionViewModel::GetSectionData() const
+const UHierarchySection* FHierarchySectionViewModel::GetSectionData() const
 {
 	if(IsAllSection())
 	{
 		return nullptr;
 	}
 
-	return GetDataMutable<UHierarchySection>();
+	return GetData<UHierarchySection>();
 }
 
 FString FHierarchySectionViewModel::ToString() const
@@ -1588,14 +1586,14 @@ bool FHierarchySectionViewModel::CanDeleteInternal()
 	return IsForHierarchy() && !IsAllSection();
 }
 
-bool FHierarchySectionViewModel::OnDroppedOnOverride(TSharedPtr<FHierarchyElementViewModel> DroppedItem, EItemDropZone ItemDropZone)
+bool FHierarchySectionViewModel::OnDroppedOnOverride(TSharedPtr<FHierarchyElementViewModel> DroppedElementViewModel, EItemDropZone ItemDropZone)
 {
-	if(DroppedItem->GetData()->IsA<UHierarchySection>())
+	if(DroppedElementViewModel->GetData()->IsA<UHierarchySection>())
 	{
 		FScopedTransaction Transaction(LOCTEXT("Transaction_OnSectionMoved", "Moved section"));
 		HierarchyViewModel->GetHierarchyRoot()->Modify();
 	
-		UHierarchySection* DraggedSectionData = DroppedItem->GetDataMutable<UHierarchySection>();
+		UHierarchySection* DraggedSectionData = DroppedElementViewModel->GetDataMutable<UHierarchySection>();
 	
 		int32 IndexOfThis = HierarchyViewModel->GetHierarchyRoot()->GetSectionData().Find(GetDataMutable<UHierarchySection>());
 		int32 DraggedSectionIndex = HierarchyViewModel->GetHierarchyRoot()->GetSectionData().Find(DraggedSectionData);
@@ -1626,7 +1624,6 @@ bool FHierarchySectionViewModel::OnDroppedOnOverride(TSharedPtr<FHierarchyElemen
 			}
 	
 			bDropSucceeded = true;
-	
 		}
 	
 		if(bDropSucceeded)
@@ -1634,30 +1631,37 @@ bool FHierarchySectionViewModel::OnDroppedOnOverride(TSharedPtr<FHierarchyElemen
 			HierarchyViewModel->ForceFullRefresh();
 			HierarchyViewModel->OnHierarchyChanged().Broadcast();
 		}
+		
+		return true;
 	}
-	else if(UHierarchyCategory* HierarchyCategory = DroppedItem->GetDataMutable<UHierarchyCategory>())
+
+	if(UHierarchyElement* DroppedElement = DroppedElementViewModel->GetDataMutable())
 	{
-		FScopedTransaction Transaction(LOCTEXT("Transaction_OnSectionDrop", "Moved category to section"));
+		FScopedTransaction Transaction(LOCTEXT("Transaction_OnSectionDrop", "Moved element to section"));
 		HierarchyViewModel->GetHierarchyRoot()->Modify();
 		
-		HierarchyCategory->SetSection(IsAllSection() ? nullptr : GetDataMutable<UHierarchySection>());
-	
-		// we null out any sections for all contained categories
-		TArray<UHierarchyCategory*> AllChildCategories;
-		HierarchyCategory->GetChildrenOfType<UHierarchyCategory>(AllChildCategories, true);
-		for(UHierarchyCategory* ChildCategory : AllChildCategories)
+		// If we are in the all section, we clear the section association
+		if(IsAllSection())
 		{
-			ChildCategory->SetSection(nullptr);
+			DroppedElement->DeleteMetaDataOfType<FDataHierarchyElementMetaData_SectionAssociation>();
 		}
-	
+		// If not, we update the section association
+		else
+		{
+			FDataHierarchyElementMetaData_SectionAssociation* SectionAssociation = DroppedElement->FindOrAddMetaDataOfType<FDataHierarchyElementMetaData_SectionAssociation>();
+			SectionAssociation->Section = GetSectionData();
+		}
+		
 		// we only need to reparent if the parent isn't already the root. This stops unnecessary reordering
-		if(DroppedItem->GetParent() != HierarchyViewModel->GetHierarchyRootViewModel())
+		if(DroppedElementViewModel->GetParent() != HierarchyViewModel->GetHierarchyRootViewModel())
 		{
-			HierarchyViewModel->GetHierarchyRootViewModel()->ReparentToThis(DroppedItem);
+			HierarchyViewModel->GetHierarchyRootViewModel()->ReparentToThis(DroppedElementViewModel);
 		}
 		
 		HierarchyViewModel->RefreshHierarchyView();
 		HierarchyViewModel->OnHierarchyChanged().Broadcast();
+		
+		return true;
 	}
 
 	return true;
@@ -1676,7 +1680,7 @@ FHierarchyElementViewModel::FResultWithUserFeedback FHierarchySectionViewModel::
 	return Results;
 }
 
-FHierarchyElementViewModel::FResultWithUserFeedback FHierarchySectionViewModel::CanDropOnInternal(TSharedPtr<FHierarchyElementViewModel> DraggedElement, EItemDropZone ItemDropZone)
+FHierarchyElementViewModel::FResultWithUserFeedback FHierarchySectionViewModel::CanDropOnInternal(TSharedPtr<FHierarchyElementViewModel> DraggedElementViewModel, EItemDropZone ItemDropZone)
 {
 	FResultWithUserFeedback Results(false);
 	// we don't allow dropping onto source sections and we don't specify a message as the sections aren't going to light up as valid drop targets
@@ -1685,20 +1689,9 @@ FHierarchyElementViewModel::FResultWithUserFeedback FHierarchySectionViewModel::
 		return false;
 	}
 	
-	if(const UHierarchyCategory* DraggedCategory = Cast<UHierarchyCategory>(DraggedElement->GetData()))
-	{
-		if(ItemDropZone == EItemDropZone::OntoItem)
-		{
-			FText Message = LOCTEXT("DropCategoryOnSectionDragMessage", "Add {0} to section {1}");
-			Message = FText::FormatOrdered(Message, FText::FromString(DraggedElement->ToString()), FText::FromString(ToString()));
-			
-			Results.bResult = GetData() != DraggedCategory->GetSection();
-			Results.UserFeedback = Results.bResult ? Message : FText::GetEmpty();
-		}
-	}
-	else if(UHierarchySection* DraggedSection = Cast<UHierarchySection>(DraggedElement->GetDataMutable()))
+	if(UHierarchySection* DraggedSection = Cast<UHierarchySection>(DraggedElementViewModel->GetDataMutable()))
 	{
-		TSharedPtr<FHierarchySectionViewModel> DraggedSectionViewModel = StaticCastSharedPtr<FHierarchySectionViewModel>(DraggedElement);
+		TSharedPtr<FHierarchySectionViewModel> DraggedSectionViewModel = StaticCastSharedPtr<FHierarchySectionViewModel>(DraggedElementViewModel);
 		const bool bSameSection = GetSectionData() == DraggedSectionViewModel->GetSectionData();
 
 		// If we drag a section onto a section, nothing happens
@@ -1729,14 +1722,29 @@ FHierarchyElementViewModel::FResultWithUserFeedback FHierarchySectionViewModel::
 			if(ItemDropZone != EItemDropZone::OntoItem)
 			{
 				FText Message = LOCTEXT("MoveSectionLeftDragMessage", "Move section {0} here");
-				Message = FText::FormatOrdered(Message, FText::FromString(DraggedElement->ToString()));
+				Message = FText::FormatOrdered(Message, FText::FromString(DraggedElementViewModel->ToString()));
 				Results.UserFeedback = Message;
 			}
 		}
 	}
-	else if(UHierarchyItem* Item = Cast<UHierarchyItem>(DraggedElement->GetDataMutable()))
+	// For any element other than sections, we let the CanContain check handle that will be called before this filter out invalid element types
+	// If it's a valid type, we assume it's generally allowed to be dropped onto a section
+	else
 	{
-		return false;
+		if(ItemDropZone == EItemDropZone::OntoItem)
+		{
+			// We only allow adding to a section if the element isn't already part of that section
+			FDataHierarchyElementMetaData_SectionAssociation SectionAssociation = DraggedElementViewModel->GetData()->FindMetaDataOfTypeOrDefault<FDataHierarchyElementMetaData_SectionAssociation>();
+			
+			Results.bResult = SectionAssociation.Section != GetSectionData();
+			if(Results.bResult == true)
+			{
+				FText Message = LOCTEXT("AssociateElementWithSectionDragMessage", "Associate element {0} with section {1}");
+				Results.UserFeedback = FText::FormatOrdered(Message, FText::FromString(DraggedElementViewModel->ToString()), GetSectionNameAsText()); 
+			}
+		}
+		
+		return Results;
 	}
 
 	return Results;
@@ -1748,42 +1756,18 @@ void FHierarchySectionViewModel::FinalizeInternal()
 	{
 		HierarchyViewModel->SetActiveHierarchySection(HierarchyViewModel->GetDefaultHierarchySectionViewModel());
 	}
-
-	// we make sure to reset all categories' section entry that were referencing this section
-	TArray<UHierarchyCategory*> AllCategories;
-	HierarchyViewModel->GetHierarchyRoot()->GetChildrenOfType<UHierarchyCategory>(AllCategories, true);
-
-	for(UHierarchyCategory* Category : AllCategories)
-	{
-		if(Category->GetSection() == GetData())
-		{
-			Category->SetSection(nullptr);
-		}
-	}
 }
 
 FHierarchyElementViewModel::FResultWithUserFeedback FHierarchyItemViewModel::CanContainInternal(TSubclassOf<UHierarchyElement> InHierarchyElement)
 {
+	// By default, items can't contain other children
 	return false;
 }
 
-::FHierarchyElementViewModel::FResultWithUserFeedback FHierarchyItemViewModel::CanDropOnInternal(TSharedPtr<FHierarchyElementViewModel> DraggedElement, EItemDropZone ItemDropZone)
+FHierarchyElementViewModel::FResultWithUserFeedback FHierarchyItemViewModel::CanDropOnInternal(TSharedPtr<FHierarchyElementViewModel> DraggedElement, EItemDropZone ItemDropZone)
 {
-	bool bAllowDrop = false;
-	
-	TSharedPtr<FHierarchyElementViewModel> SourceDropItem = DraggedElement;
-	TSharedPtr<FHierarchyElementViewModel> TargetDropItem = AsShared();
-
-	// we only allow drops if some general conditions are fulfilled
-	if(SourceDropItem->GetData() != TargetDropItem->GetData() &&
-		(!SourceDropItem->HasParent(TargetDropItem, false) || ItemDropZone != EItemDropZone::OntoItem)  &&
-		!TargetDropItem->HasParent(SourceDropItem, true))
-	{
-		// items can be generally be dropped above/below other items
-		bAllowDrop = (SourceDropItem->GetData()->IsA<UHierarchyItem>() && ItemDropZone != EItemDropZone::OntoItem);
-	}
-
-	return bAllowDrop;
+	// items can be generally be dropped above/below other items
+	return (GetData()->IsA<UHierarchyItem>() && ItemDropZone != EItemDropZone::OntoItem);
 }
 
 const FTableRowStyle* FHierarchyItemViewModel::GetRowStyle() const
@@ -1791,39 +1775,6 @@ const FTableRowStyle* FHierarchyItemViewModel::GetRowStyle() const
 	return &FAppStyle::Get().GetWidgetStyle<FTableRowStyle>("TableView.Row");
 }
 
-void FHierarchyCategoryViewModel::PostAddFixup()
-{
-	// If the new category is under the root, we set the section to the currently active one
-	if(Parent.Pin()->GetData()->IsA<UHierarchyRoot>())
-	{
-		UHierarchyCategory* Category = GetDataMutable<UHierarchyCategory>();
-		Category->SetSection(HierarchyViewModel->GetActiveHierarchySectionViewModel()->IsAllSection() ? nullptr : HierarchyViewModel->GetActiveHierarchySectionData());
-
-		// we null out any sections for all contained categories
-		TArray<UHierarchyCategory*> AllChildCategories;
-		Category->GetChildrenOfType<UHierarchyCategory>(AllChildCategories, true);
-		for(UHierarchyCategory* ChildCategory : AllChildCategories)
-		{
-			ChildCategory->SetSection(nullptr);
-		}
-	}
-	// If it's a subcategory, we empty out the section
-	else if(Parent.Pin()->GetData()->IsA<UHierarchyCategory>())
-	{
-		UHierarchyCategory* Category = GetDataMutable<UHierarchyCategory>();
-
-		Category->SetSection(nullptr);
-
-		// we null out any sections for all contained categories
-		TArray<UHierarchyCategory*> AllChildCategories;
-		Category->GetChildrenOfType<UHierarchyCategory>(AllChildCategories, true);
-		for(UHierarchyCategory* ChildCategory : AllChildCategories)
-		{
-			ChildCategory->SetSection(nullptr);
-		}
-	}
-}
-
 const FTableRowStyle* FHierarchyCategoryViewModel::GetRowStyle() const
 {
 	return &FDataHierarchyEditorStyle::Get().GetWidgetStyle<FTableRowStyle>("HierarchyEditor.Row.Category");
@@ -1852,13 +1803,12 @@ void UDataHierarchyViewModelBase::AddCategory(TSharedPtr<FHierarchyElementViewMo
 
 	UClass* CategoryClass = GetCategoryDataClass();
 	
-	UHierarchyCategory* Category = Cast<UHierarchyCategory>(CategoryParent->AddChild(CategoryClass, UHierarchyCategory::ConstructIdentity()));
-	
-	TSharedPtr<FHierarchyElementViewModel> ViewModel = CategoryParent->FindViewModelForChild(Category, false);
-	if(ensureMsgf(ViewModel.IsValid(), TEXT("Could not find view model for new category of type '%s'. Please ensure your 'CreateViewModelForData' function creates a view model."), *CategoryClass->GetName()))
+	TSharedPtr<FHierarchyCategoryViewModel> CategoryViewModel = StaticCastSharedPtr<FHierarchyCategoryViewModel>(CategoryParent->AddChild(CategoryClass, INDEX_NONE, UHierarchyCategory::ConstructIdentity()));
+	if(ensureMsgf(CategoryViewModel.IsValid(), TEXT("Could not find view model for new category of type '%s'. Please ensure your 'CreateViewModelForData' function creates a view model."), *CategoryClass->GetName()))
 	{		
 		TArray<UHierarchyCategory*> SiblingCategories;
-		Category->GetTypedOuter<UHierarchyElement>()->GetChildrenOfType<UHierarchyCategory>(SiblingCategories);
+
+		CategoryViewModel->GetParent().Pin()->GetData()->GetChildrenOfType<UHierarchyCategory>(SiblingCategories);
 		
 		TSet<FName> CategoryNames;
 		for(const auto& SiblingCategory : SiblingCategories)
@@ -1866,24 +1816,14 @@ void UDataHierarchyViewModelBase::AddCategory(TSharedPtr<FHierarchyElementViewMo
 			CategoryNames.Add(SiblingCategory->GetCategoryName());
 		}
 
-		Category->SetCategoryName(UE::DataHierarchyEditor::GetUniqueName(FName("New Category"), CategoryNames));
-		
-		// we only set the section property if the current section isn't set to "All"
-		Category->SetSection(GetActiveHierarchySectionViewModel()->IsAllSection() ? nullptr : GetActiveHierarchySectionData());
+		CategoryViewModel->SetCategoryName(UE::DataHierarchyEditor::GetUniqueName(FName("New Category"), CategoryNames));
 		
 		RefreshHierarchyView();
-
-		OnElementAddedDelegate.ExecuteIfBound(ViewModel);
+		
+		//OnElementAddedDelegate.ExecuteIfBound(ViewModel);
 	}
 }
 
-void UDataHierarchyViewModelBase::AddSection() const
-{
-	TSharedPtr<FHierarchySectionViewModel> SectionViewModel = GetHierarchyRootViewModel()->AddSection();	
-	OnElementAddedDelegate.ExecuteIfBound(SectionViewModel);
-	OnHierarchyChangedDelegate.Broadcast();
-}
-
 void UDataHierarchyViewModelBase::GenerateDynamicContextMenu(UToolMenu* ToolMenu)
 {
 	UHierarchyMenuContext* HierarchyMenuContext = ToolMenu->FindContext<UHierarchyMenuContext>();
@@ -1956,11 +1896,11 @@ void UDataHierarchyViewModelBase::GenerateDynamicContextMenuInternal(UToolMenu*
 	}
 }
 
-UHierarchyElement* UDataHierarchyViewModelBase::AddElementUnderRoot(TSubclassOf<UHierarchyElement> NewChildClass, FHierarchyElementIdentity ChildIdentity)
+void UDataHierarchyViewModelBase::AddElementUnderRoot(TSubclassOf<UHierarchyElement> NewChildClass, FHierarchyElementIdentity ChildIdentity)
 {
 	FScopedTransaction Transaction(LOCTEXT("Transaction_AddItem", "Add hierarchy item"));
 	HierarchyRoot->Modify();
-	return GetHierarchyRootViewModel()->AddChild(NewChildClass, ChildIdentity);
+	GetHierarchyRootViewModel()->AddChild(NewChildClass, INDEX_NONE, ChildIdentity);
 }
 
 void UDataHierarchyViewModelBase::DeleteElementWithIdentity(FHierarchyElementIdentity Identity)
@@ -1983,7 +1923,9 @@ void UDataHierarchyViewModelBase::DeleteElementWithIdentity(FHierarchyElementIde
 		}
 	}
 	
-	TArray<TSharedPtr<FHierarchySectionViewModel>> SectionViewModels = HierarchyRootViewModel->GetSectionViewModels();
+	TArray<TSharedPtr<FHierarchySectionViewModel>> SectionViewModels; 
+	HierarchyRootViewModel->GetChildrenViewModelsForType<UHierarchySection, FHierarchySectionViewModel>(SectionViewModels);
+	
 	for(TSharedPtr<FHierarchySectionViewModel> SectionViewModel : SectionViewModels)
 	{
 		if(SectionViewModel->GetData()->GetPersistentIdentity() == Identity && SectionViewModel->CanDelete())
diff --git a/Engine/Source/Editor/DataHierarchyEditor/Private/Widgets/SDataHierarchyEditor.cpp b/Engine/Source/Editor/DataHierarchyEditor/Private/Widgets/SDataHierarchyEditor.cpp
--- a/Engine/Source/Editor/DataHierarchyEditor/Private/Widgets/SDataHierarchyEditor.cpp
+++ b/Engine/Source/Editor/DataHierarchyEditor/Private/Widgets/SDataHierarchyEditor.cpp
@@ -124,6 +124,7 @@ void SHierarchySection::Construct(const FArguments& InArgs, TSharedPtr<FHierarch
 	[
 		SNew(SHorizontalBox)
 		+ SHorizontalBox::Slot()
+		.VAlign(VAlign_Fill)
 		[
 			SAssignNew(CheckBox, SCheckBox)
 			.Visibility(EVisibility::HitTestInvisible)
@@ -137,7 +138,14 @@ void SHierarchySection::Construct(const FArguments& InArgs, TSharedPtr<FHierarch
 				.AutoWidth()
 				.Padding(2.f)
 				[
-					SNew(SImage).Image(InSection->GetSectionImage() != nullptr ? InSection->GetSectionImage() : FStyleDefaults::GetNoBrush())
+					SNew(SBox)
+					.WidthOverride(20.f)
+					.HeightOverride(20.f)
+					.Visibility(this, &SHierarchySection::GetImageVisibility)
+					[
+						SNew(SImage)
+						.Image(this, &SHierarchySection::GetImageBrush)
+					]
 				]
 				+ SHorizontalBox::Slot()
 				[
@@ -374,6 +382,16 @@ UHierarchySection* SHierarchySection::TryGetSectionData() const
 	return Cast<UHierarchySection>(SectionViewModelWeak.Pin()->GetDataMutable());
 }
 
+const FSlateBrush* SHierarchySection::GetImageBrush() const
+{
+	return (SectionViewModelWeak.IsValid() && SectionViewModelWeak.Pin()->GetSectionImageBrush() != FAppStyle::GetNoBrush()) ? SectionViewModelWeak.Pin()->GetSectionImageBrush() : FAppStyle::GetNoBrush();
+}
+
+EVisibility SHierarchySection::GetImageVisibility() const
+{
+	return (SectionViewModelWeak.IsValid() && SectionViewModelWeak.Pin()->GetSectionImageBrush() != FAppStyle::GetNoBrush()) ? EVisibility::Visible : EVisibility::Collapsed;  
+}
+
 FText SHierarchySection::GetText() const
 {
 	return SectionViewModelWeak.IsValid() ? SectionViewModelWeak.Pin()->GetSectionNameAsText() : FText::GetEmpty();
@@ -521,7 +539,7 @@ void SDataHierarchyEditor::Construct(const FArguments& InArgs, TObjectPtr<UDataH
 	HierarchyViewModel->OnRefreshHierarchyView().BindSP(this, &SDataHierarchyEditor::RefreshHierarchyView);
 	HierarchyViewModel->OnRefreshSectionsView().BindSP(this, &SDataHierarchyEditor::RefreshSectionsView);
 	HierarchyViewModel->OnHierarchySectionActivated().BindSP(this, &SDataHierarchyEditor::OnHierarchySectionActivated);
-	HierarchyViewModel->OnElementAdded().BindSP(this, &SDataHierarchyEditor::OnElementAdded);
+	HierarchyViewModel->OnHierarchyElementChanged().BindSP(this, &SDataHierarchyEditor::OnHierarchyElementChanged);
 
 	BindToHierarchyRootViewModel();
 	
@@ -837,7 +855,7 @@ SDataHierarchyEditor::~SDataHierarchyEditor()
 		HierarchyViewModel->OnRefreshHierarchyView().Unbind();
 		HierarchyViewModel->OnRefreshSectionsView().Unbind();
 		HierarchyViewModel->OnHierarchySectionActivated().Unbind();
-		HierarchyViewModel->OnElementAdded().Unbind();
+		HierarchyViewModel->OnHierarchyElementChanged().Unbind();
 
 		UnbindFromHierarchyRootViewModel();
 
@@ -925,7 +943,7 @@ FReply SDataHierarchyEditor::OnAddCategoryClicked() const
 
 FReply SDataHierarchyEditor::OnAddSectionClicked() const
 {
-	HierarchyViewModel.Get()->AddSection();
+	HierarchyViewModel.Get()->GetHierarchyRootViewModel()->AddSection();
 	return FReply::Handled();
 }
 
@@ -965,6 +983,8 @@ void SDataHierarchyEditor::RefreshSourceView(bool bFullRefresh) const
 	{
 		SourceTreeView->RequestTreeRefresh();
 	}
+	
+	ExpandEntriesByDefault(SourceTreeView);
 }
 
 void SDataHierarchyEditor::RequestRefreshSourceViewNextFrame(bool bFullRefresh)
@@ -997,6 +1017,8 @@ void SDataHierarchyEditor::RefreshHierarchyView(bool bFullRefresh) const
 	{
 		HierarchyTreeView->RequestTreeRefresh();
 	}
+	
+	ExpandEntriesByDefault(HierarchyTreeView);
 }
 
 void SDataHierarchyEditor::RequestRefreshHierarchyViewNextFrame(bool bFullRefresh)
@@ -1033,7 +1055,7 @@ void SDataHierarchyEditor::RefreshSectionsView()
 			SourceSectionBox->AddSlot()
 			.Padding(2.f)
 			.HAlign(HAlign_Center)
-			.VAlign(VAlign_Center)
+			.VAlign(VAlign_Fill)
 			[
 				SectionWidget.ToSharedRef()		
 			];
@@ -1051,7 +1073,7 @@ void SDataHierarchyEditor::RefreshSectionsView()
 			SourceSectionBox->AddSlot()
 			.Padding(2.f)
 			.HAlign(HAlign_Center)
-			.VAlign(VAlign_Center)
+			.VAlign(VAlign_Fill)
 			[
 				DefaultSourceSection.ToSharedRef()
 			];
@@ -1059,7 +1081,7 @@ void SDataHierarchyEditor::RefreshSectionsView()
 	}
 
 	HierarchySectionBox->ClearChildren();
-	for (TSharedPtr<FHierarchySectionViewModel>& HierarchySection : HierarchyViewModel->GetHierarchyRootViewModel()->GetSectionViewModels())
+	for (const TSharedPtr<FHierarchySectionViewModel>& HierarchySection : HierarchyViewModel->GetHierarchyRootViewModel()->GetSectionViewModels())
 	{
 		TSharedPtr<SHierarchySection> SectionWidget = SNew(SHierarchySection, HierarchySection)
 		.IsSectionActive_Lambda([this, HierarchySection]()
@@ -1074,7 +1096,7 @@ void SDataHierarchyEditor::RefreshSectionsView()
 		HierarchySectionBox->AddSlot()
 		.Padding(2.f)
 		.HAlign(HAlign_Center)
-		.VAlign(VAlign_Center)
+		.VAlign(VAlign_Fill)
 		[
 			SectionWidget.ToSharedRef()		
 		];
@@ -1093,7 +1115,7 @@ void SDataHierarchyEditor::RefreshSectionsView()
 	HierarchySectionBox->AddSlot()
 	.Padding(2.f)
 	.HAlign(HAlign_Center)
-	.VAlign(VAlign_Center)
+	.VAlign(VAlign_Fill)
 	[
 		DefaultHierarchySection.ToSharedRef()
 	];
@@ -1239,7 +1261,7 @@ TSharedRef<ITableRow> SDataHierarchyEditor::GenerateHierarchyItemRow(TSharedPtr<
 	];
 }
 
-bool SDataHierarchyEditor::FilterForSourceSection(TSharedPtr<const FHierarchyElementViewModel> ItemViewModel) const
+bool SDataHierarchyEditor::FilterForSourceSection(TSharedPtr<const FHierarchyElementViewModel> ViewModel) const
 {
 	if(ActiveSourceSection.IsValid())
 	{
@@ -1250,7 +1272,8 @@ bool SDataHierarchyEditor::FilterForSourceSection(TSharedPtr<const FHierarchyEle
 		}
 
 		// if not, we check against identical section data
-		return ActiveSourceSection.Pin()->GetData() == ItemViewModel->GetSection();
+		FDataHierarchyElementMetaData_SectionAssociation SectionAssociation = ViewModel->GetData()->FindMetaDataOfTypeOrDefault<FDataHierarchyElementMetaData_SectionAssociation>();		
+		return ActiveSourceSection.Pin()->GetSectionData() == SectionAssociation.Section;
 	}
 
 	return true;
@@ -1439,21 +1462,52 @@ bool SDataHierarchyEditor::CanNavigateToMatchingHierarchyElementFromSelectedSour
 	return false;
 }
 
-void SDataHierarchyEditor::OnElementAdded(TSharedPtr<FHierarchyElementViewModel> AddedItem)
+void SDataHierarchyEditor::OnHierarchyElementChanged(TInstancedStruct<FHierarchyElementChangedPayload> Payload)
 {
-	// when a new item is created (opposed to dragged & dropped from source view, i.e. only categories so far)
-	// we make sure to request a tree refresh, select the row, and request a pending rename since the widget will created a frame later
-	if(AddedItem->GetData()->IsA<UHierarchyItem>() || AddedItem->GetData()->IsA<UHierarchyCategory>())
-	{
-		HierarchyTreeView->RequestTreeRefresh();
-		NavigateToHierarchyElement(AddedItem);
-	}
-	else if(AddedItem->GetData()->IsA<UHierarchySection>())
+	if(Payload.IsValid())
 	{
-		RefreshSectionsView();
-	}
+		if(Payload.GetScriptStruct() == FHierarchyElementChangedPayload_AddedElement::StaticStruct())
+		{
+			TSharedPtr<FHierarchyElementViewModel> AddedElementViewModel = Payload.GetMutablePtr<FHierarchyElementChangedPayload_AddedElement>()->AddedElementViewModel.Pin();
+			
+			// when a new item is created (opposed to dragged & dropped from source view, i.e. only categories so far)
+			// we make sure to request a tree refresh, select the row, and request a pending rename since the widget will created a frame later
+			if(AddedElementViewModel->GetData()->IsA<UHierarchyItem>() || AddedElementViewModel->GetData()->IsA<UHierarchyCategory>())
+			{
+				HierarchyTreeView->RequestTreeRefresh();
+				NavigateToHierarchyElement(AddedElementViewModel);
+			}
+			else if(AddedElementViewModel->GetData()->IsA<UHierarchySection>())
+			{
+				RefreshSectionsView();
+				HierarchyViewModel->SetActiveHierarchySection(StaticCastSharedPtr<FHierarchySectionViewModel>(AddedElementViewModel));
+			}
+
+			AddedElementViewModel->RequestRenamePending();
+		}
 
-	AddedItem->RequestRenamePending();
+		if(Payload.GetScriptStruct() == FHierarchyElementChangedPayload_DeletedElement::StaticStruct())
+		{
+			TSharedPtr<FHierarchyElementViewModel> DeletedElementViewModel = Payload.GetMutablePtr<FHierarchyElementChangedPayload_DeletedElement>()->DeletedElementViewModel.Pin();
+			
+			// when a new item is created (opposed to dragged & dropped from source view, i.e. only categories so far)
+			// we make sure to request a tree refresh, select the row, and request a pending rename since the widget will created a frame later
+			if(DeletedElementViewModel->GetData()->IsA<UHierarchyItem>() || DeletedElementViewModel->GetData()->IsA<UHierarchyCategory>())
+			{
+				HierarchyTreeView->RequestTreeRefresh();
+				NavigateToHierarchyElement(DeletedElementViewModel);
+			}
+			else if(DeletedElementViewModel->GetData()->IsA<UHierarchySection>())
+			{
+				RefreshSectionsView();
+
+				if(HierarchyViewModel->GetActiveHierarchySectionViewModel() == DeletedElementViewModel)
+				{
+					HierarchyViewModel->SetActiveHierarchySection(HierarchyViewModel->GetDefaultHierarchySectionViewModel());
+				}
+			}
+		}
+	}
 }
 
 void SDataHierarchyEditor::OnHierarchySectionActivated(TSharedPtr<FHierarchySectionViewModel> Section)
@@ -1524,7 +1578,15 @@ TSharedRef<SWidget> SDataHierarchyEditor::OnGetCustomAddContent()
 
 					if(ParentViewModel->CanContain(Type).bResult)
 					{
-						ParentViewModel->AddChild(Type, FHierarchyElementIdentity());
+						TSharedPtr<FHierarchyElementViewModel> NewElementViewModel = ParentViewModel->AddChild(Type);
+
+						if(const UHierarchySection* Section = DataHierarchyEditor->HierarchyViewModel->GetActiveHierarchySectionData())
+						{
+							FDataHierarchyElementMetaData_SectionAssociation* SectionAssociation = NewElementViewModel->GetDataMutable()->FindOrAddMetaDataOfType<FDataHierarchyElementMetaData_SectionAssociation>();
+							SectionAssociation->Section = Section;
+						}
+
+						NewElementViewModel->SyncViewModelsToData();
 					}
 				});
 				UIAction.CanExecuteAction = FToolMenuCanExecuteAction::CreateLambda([ContextObject, Type](const FToolMenuContext& Context)
@@ -1866,6 +1928,26 @@ TOptional<SSearchBox::FSearchResultData> SDataHierarchyEditor::GetSearchResultDa
 	return TOptional<SSearchBox::FSearchResultData>();
 }
 
+void SDataHierarchyEditor::ExpandEntriesByDefault(TSharedPtr<STreeView<TSharedPtr<FHierarchyElementViewModel>>> TreeView) const
+{
+	TArray<TSharedPtr<FHierarchyElementViewModel>> EntriesToProcess(TreeView->GetRootItems());
+	while (EntriesToProcess.Num() > 0)
+	{
+		TSharedPtr<FHierarchyElementViewModel> EntryToProcess = EntriesToProcess[0];
+		EntriesToProcess.RemoveAtSwap(0);
+
+		if (EntryToProcess->IsExpandedByDefault())
+		{
+			TreeView->SetItemExpansion(EntryToProcess, true);
+			EntriesToProcess.Append(EntryToProcess->GetFilteredChildren());
+		}
+		else
+		{
+			TreeView->SetItemExpansion(EntryToProcess, false);
+		}
+	}
+}
+
 FReply SDataHierarchyEditor::HandleHierarchyRootDrop(const FGeometry& Geometry, const FDragDropEvent& DragDropEvent) const
 {
 	if(TSharedPtr<FHierarchyDragDropOp> DragDropOp = DragDropEvent.GetOperationAs<FHierarchyDragDropOp>())
@@ -1928,8 +2010,53 @@ FSlateColor SDataHierarchyEditor::GetRootIconColor() const
 	return FLinearColor(0.2f, 0.2f, 0.2f, 0.5f);
 }
 
-void SDataHierarchyEditor::NotifyPostChange(const FPropertyChangedEvent& PropertyChangedEvent,	FProperty* PropertyThatChanged)
+void SDataHierarchyEditor::NotifyPostChange(const FPropertyChangedEvent& PropertyChangedEvent, FProperty* PropertyThatChanged)
 {
+	const UObject* EditedObject = PropertyChangedEvent.GetObjectBeingEdited(0);
+
+	// We look for the view model corresponding to the edited object.
+	// The object doesn't have to be a hierarchy element, since view models can provide external elements for editing. 
+	if(EditedObject->IsA<UHierarchySection>())
+	{
+		if(SourceRootViewModel.IsValid())
+		{
+			TSharedPtr<FHierarchySectionViewModel>* FoundSourceSectionViewModel = SourceRootViewModel->GetSectionViewModels().FindByPredicate([EditedObject](TSharedPtr<FHierarchySectionViewModel> Candidate)
+				{
+					return Candidate->GetSectionData() == EditedObject;
+				});
+			
+			if(FoundSourceSectionViewModel)
+			{
+				(*FoundSourceSectionViewModel)->PostEditChange(PropertyChangedEvent, PropertyThatChanged);
+			}
+		}
+
+		TSharedPtr<FHierarchySectionViewModel>* FoundHierarchySectionViewModel = HierarchyViewModel->GetHierarchyRootViewModel()->GetSectionViewModels().FindByPredicate([EditedObject](TSharedPtr<FHierarchySectionViewModel> Candidate)
+			{
+				return Candidate->GetSectionData() == EditedObject;
+			});
+
+		if(FoundHierarchySectionViewModel)
+		{
+			(*FoundHierarchySectionViewModel)->PostEditChange(PropertyChangedEvent, PropertyThatChanged);
+		}
+	}
+	else if(EditedObject->IsA<UHierarchyElement>())
+	{
+		if(SourceRootViewModel.IsValid())
+		{
+			if(TSharedPtr<FHierarchyElementViewModel> FoundSourceViewModel = SourceRootViewModel->FindViewModelForChild(Cast<UHierarchyElement>(PropertyChangedEvent.GetObjectBeingEdited(0)), true))
+			{
+				FoundSourceViewModel->PostEditChange(PropertyChangedEvent, PropertyThatChanged);
+			}
+		}
+		
+		if(TSharedPtr<FHierarchyElementViewModel> FoundHierarchyViewModel = HierarchyViewModel->GetHierarchyRootViewModel()->FindViewModelForChild(Cast<UHierarchyElement>(PropertyChangedEvent.GetObjectBeingEdited(0)), true))
+		{
+			FoundHierarchyViewModel->PostEditChange(PropertyChangedEvent, PropertyThatChanged);
+		}
+	}
+	
 	HierarchyViewModel->OnHierarchyPropertiesChanged().Broadcast();
 }
 
diff --git a/Engine/Source/Editor/DataHierarchyEditor/Public/DataHierarchyViewModelBase.h b/Engine/Source/Editor/DataHierarchyEditor/Public/DataHierarchyViewModelBase.h
--- a/Engine/Source/Editor/DataHierarchyEditor/Public/DataHierarchyViewModelBase.h
+++ b/Engine/Source/Editor/DataHierarchyEditor/Public/DataHierarchyViewModelBase.h
@@ -14,6 +14,8 @@
 #include "Templates/SubclassOf.h"
 #include "UObject/Package.h"
 #include "UObject/UObjectGlobals.h"
+#include "StructUtils/InstancedStruct.h"
+#include "DataHierarchyCommonTypes.h"
 #include "DataHierarchyViewModelBase.generated.h"
 
 #define UE_API DATAHIERARCHYEDITOR_API
@@ -61,80 +63,6 @@
  */
 
 struct FHierarchyRootViewModel;
-/** This struct is used to identify a given hierarchy element and can be based on guids and/or names.
- *  This is particularly useful when a hierarchy element represents an object or a property that is not owned by the hierarchy itself.
- *  
- *  
- */
-USTRUCT()
-struct FHierarchyElementIdentity
-{
-	GENERATED_BODY()
-
-	FHierarchyElementIdentity() {}
-	FHierarchyElementIdentity(TArray<FGuid> InGuids, TArray<FName> InNames) : Guids(InGuids), Names(InNames) {}
-	
-	/** An array of guids that have to be satisfied in order to match. */
-	UPROPERTY()
-	TArray<FGuid> Guids;
-
-	/** Optionally, an array of names can be specified in place of guids. If guids & names are present, guids have to be satisfied first, then names. */
-	UPROPERTY()
-	TArray<FName> Names;
-
-	bool IsValid() const
-	{
-		return Guids.Num() > 0 || Names.Num() > 0;
-	}
-	
-	bool operator==(const FHierarchyElementIdentity& OtherIdentity) const
-	{
-		if(Guids.Num() != OtherIdentity.Guids.Num() || Names.Num() != OtherIdentity.Names.Num())
-		{
-			return false;
-		}
-
-		for(int32 GuidIndex = 0; GuidIndex < Guids.Num(); GuidIndex++)
-		{
-			if(Guids[GuidIndex] != OtherIdentity.Guids[GuidIndex])
-			{
-				return false;
-			}
-		}
-
-		for(int32 NameIndex = 0; NameIndex < Names.Num(); NameIndex++)
-		{
-			if(!Names[NameIndex].IsEqual(OtherIdentity.Names[NameIndex]))
-			{
-				return false;
-			}
-		}
-
-		return true;
-	}
-
-	bool operator!=(const FHierarchyElementIdentity& OtherIdentity) const
-	{
-		return !(*this == OtherIdentity);
-	}
-};
-
-FORCEINLINE uint32 GetTypeHash(const FHierarchyElementIdentity& Identity)
-{
-	uint32 Hash = 0;
-	
-	for(const FGuid& Guid : Identity.Guids)
-	{
-		Hash = HashCombine(Hash, GetTypeHash(Guid));
-	}
-	
-	for(const FName& Name : Identity.Names)
-	{
-		Hash = HashCombine(Hash, GetTypeHash(Name));
-	}
-	
-	return Hash;
-}
 
 /** A base class that is used to refresh data that represents external data. Inherit from this class if you need more context data. */
 UCLASS(MinimalAPI, Transient)
@@ -152,12 +80,15 @@ public:
 	UE_API UHierarchyElement();
 	virtual ~UHierarchyElement() override {}
 
+	UE_API TConstArrayView<const TObjectPtr<UHierarchyElement>> GetChildren() const;
 	TArray<TObjectPtr<UHierarchyElement>>& GetChildrenMutable() { return Children; }
-	UE_API const TArray<UHierarchyElement*>& GetChildren() const;
 
 	template<class ChildClass>
 	ChildClass* AddChild();
 
+	template<class ChildClass>
+	ChildClass* AddChild(TSubclassOf<UHierarchyElement> Class, int32 InsertIndex = INDEX_NONE);
+
 	UE_API UHierarchyElement* FindChildWithIdentity(FHierarchyElementIdentity ChildIdentity, bool bSearchRecursively = false);
 
 	UE_API UHierarchyElement* CopyAndAddItemAsChild(const UHierarchyElement& ItemToCopy);
@@ -187,7 +118,26 @@ public:
 	/** Overridden modify method to also mark all children as modified */
 	UE_API virtual bool Modify(bool bAlwaysMarkDirty = true) override;
 
+	template<typename T>
+	bool HasMetaDataOfType();
+
+	template<typename T>
+	const T* FindMetaDataOfType() const;
+	
+	template<typename T>
+	T* FindMetaDataOfType();
+	
+	template<typename T>
+	T FindMetaDataOfTypeOrDefault() const;
+    	
+	template<typename T>
+	T* FindOrAddMetaDataOfType();
+
+	template <typename T>
+	bool DeleteMetaDataOfType();
+	
 protected:
+	UE_API virtual void Serialize(FArchive& Ar) override;
 	UE_API virtual void PostLoad() override;
 
 	UPROPERTY()
@@ -195,6 +145,10 @@ protected:
 
 	UPROPERTY()
 	FHierarchyElementIdentity Identity;
+
+	/** A map of metadata. Can be used for anything; an example of this is section association. */
+	UPROPERTY()
+	TMap<TObjectPtr<UStruct>, TInstancedStruct<FDataHierarchyElementMetaData>> MetaData;
 	
 	/** An optional guid; can be used if hierarchy items represent outside items */
 	UPROPERTY()
@@ -205,11 +159,31 @@ template <class ChildClass>
 ChildClass* UHierarchyElement::AddChild()
 {
 	ChildClass* NewChild = NewObject<ChildClass>(this);
+	NewChild->SetFlags(RF_Transactional);
+
 	GetChildrenMutable().Add(NewChild);
 
 	return NewChild;
 }
 
+template <class ChildClass>
+ChildClass* UHierarchyElement::AddChild(TSubclassOf<UHierarchyElement> Class, int32 InsertIndex)
+{
+	ChildClass* NewChild = NewObject<ChildClass>(this, Class);
+	NewChild->SetFlags(RF_Transactional);
+
+	if(InsertIndex == INDEX_NONE)
+	{
+		GetChildrenMutable().Add(NewChild);
+	}
+	else
+	{
+		GetChildrenMutable().Insert(NewChild, InsertIndex);
+	}
+
+	return NewChild;
+}
+
 template <class ChildClass>
 bool UHierarchyElement::DoesOneChildExist(bool bRecursive) const
 {
@@ -269,6 +243,80 @@ void UHierarchyElement::SortChildren(const PREDICATE_CLASS& Predicate, bool bRec
 	}
 }
 
+template <typename T>
+bool UHierarchyElement::HasMetaDataOfType()
+{
+	static_assert(TIsDerivedFrom<T, FDataHierarchyElementMetaData>::IsDerived, "T must be derived from FDataHierarchyElementMetaData");
+
+	return MetaData.Contains(T::StaticStruct());
+}
+
+template <typename T>
+const T* UHierarchyElement::FindMetaDataOfType() const
+{
+	static_assert(TIsDerivedFrom<T, FDataHierarchyElementMetaData>::IsDerived, "T must be derived from FDataHierarchyElementMetaData");
+	
+	if(MetaData.Contains(T::StaticStruct()))
+	{
+		return MetaData[T::StaticStruct()].template GetPtr<T>();
+	}
+
+	return nullptr;
+}
+
+template <typename T>
+T* UHierarchyElement::FindMetaDataOfType()
+{
+	static_assert(TIsDerivedFrom<T, FDataHierarchyElementMetaData>::IsDerived, "T must be derived from FDataHierarchyElementMetaData");
+	
+	if(MetaData.Contains(T::StaticStruct()))
+	{
+		return MetaData[T::StaticStruct()].template GetMutablePtr<T>();
+	}
+
+	return nullptr;
+}
+
+template <typename T>
+T UHierarchyElement::FindMetaDataOfTypeOrDefault() const
+{
+	static_assert(TIsDerivedFrom<T, FDataHierarchyElementMetaData>::IsDerived, "T must be derived from FDataHierarchyElementMetaData");
+	static_assert(TIsConstructible<T>::Value, "T must be default constructible");
+	
+	if(MetaData.Contains(T::StaticStruct()))
+	{
+		return MetaData[T::StaticStruct()].template Get<T>();
+	}
+
+	return T();
+}
+
+template <typename T>
+T* UHierarchyElement::FindOrAddMetaDataOfType()
+{
+	static_assert(TIsDerivedFrom<T, FDataHierarchyElementMetaData>::IsDerived, "T must be derived from FDataHierarchyElementMetaData");
+	
+	if(MetaData.Contains(T::StaticStruct()))
+	{
+		return MetaData[T::StaticStruct()].template GetMutablePtr<T>();
+	}
+
+	return MetaData.FindOrAdd(T::StaticStruct(), TInstancedStruct<T>::Make()).template GetMutablePtr<T>();
+}
+
+template <typename T>
+bool UHierarchyElement::DeleteMetaDataOfType()
+{
+	static_assert(TIsDerivedFrom<T, FDataHierarchyElementMetaData>::IsDerived, "T must be derived from FDataHierarchyElementMetaData");
+
+	if(MetaData.Contains(T::StaticStruct()))
+	{
+		return MetaData.Remove(T::StaticStruct()) != 0;
+	}
+
+	return false;
+}
+
 /** A minimal implementation of a section. */
 UCLASS(MinimalAPI)
 class UHierarchySection : public UHierarchyElement
@@ -306,25 +354,22 @@ public:
 	UHierarchyRoot() {}
 	virtual ~UHierarchyRoot() override {}
 
-	const TArray<UHierarchySection*>& GetSectionData() const { return Sections; }
+	TConstArrayView<const TObjectPtr<UHierarchySection>> GetSectionData() const { return Sections; }
 	TArray<TObjectPtr<UHierarchySection>>& GetSectionDataMutable() { return Sections; }
 
-	UE_API TSet<FName> GetSections() const;
+	UE_API TSet<FName> GetSectionNames() const;
 	UE_API int32 GetSectionIndex(UHierarchySection* Section) const;
-
-	UE_API UHierarchySection* AddSection(FText InNewSectionName, int32 InsertIndex = INDEX_NONE, TSubclassOf<UHierarchySection> SectionClass = UHierarchySection::StaticClass());
+	
 	UE_API UHierarchySection* FindSectionByIdentity(FHierarchyElementIdentity SectionIdentity);
 	/** This will copy the section element itself */
 	UE_API void DuplicateSectionFromOtherRoot(const UHierarchySection& SectionToCopy);
-	UE_API void RemoveSection(FText SectionName);
-	UE_API void RemoveSectionByIdentity(FHierarchyElementIdentity SectionIdentity);
+	
+	UE_API virtual void EmptyAllData();
 	
 	UE_API virtual bool Modify(bool bAlwaysMarkDirty = true) override;
+protected:
 	UE_API virtual void Serialize(FStructuredArchive::FRecord Record) override;
-
-	UE_API virtual void EmptyAllData();
 protected:
-
 	UPROPERTY()
 	TArray<TObjectPtr<UHierarchySection>> Sections;
 };
@@ -355,21 +400,13 @@ public:
 	FText GetCategoryAsText() const { return FText::FromName(Category); }
 	FText GetTooltip() const { return Tooltip; }
 
-	void SetSection(UHierarchySection* InSection) { Section = InSection; }
-	const UHierarchySection* GetSection() const { return Section; }
-
 	virtual FString ToString() const override { return Category.ToString(); }
 
-	/** Since the category points to a section object, during merge or copy paste etc. it is possible the section pointer will point at a section from another root.
-	 *  We fix this up by looking through our available sections and match up via persistent identity.
-	 *  This function expects the correct section with the same identity to exist already at the root level
-	 */
-	UE_API void FixupSectionLinkage();
-
 	static UE_API FHierarchyElementIdentity ConstructIdentity();
 
 protected:
 	UE_API virtual void PostLoad() override;
+
 private:
 	UPROPERTY()
 	FName Category;
@@ -379,7 +416,7 @@ private:
 	FText Tooltip;
 
 	UPROPERTY()
-	TObjectPtr<UHierarchySection> Section = nullptr;
+	TObjectPtr<UHierarchySection> Section_DEPRECATED = nullptr;
 };
 
 struct FHierarchyElementViewModel;
@@ -430,6 +467,37 @@ public:
 	TArray<TSharedPtr<FHierarchyElementViewModel>> MenuHierarchyElements;
 };
 
+/** Generic payload struct for when an element in the hierarchy has changed. */
+USTRUCT()
+struct FHierarchyElementChangedPayload
+{
+	GENERATED_BODY()
+};
+
+USTRUCT()
+struct FHierarchyElementChangedPayload_AddedElement : public FHierarchyElementChangedPayload
+{
+	GENERATED_BODY()
+	
+	TWeakPtr<FHierarchyElementViewModel> AddedElementViewModel;
+};
+
+USTRUCT()
+struct FHierarchyElementChangedPayload_DeletedElement : public FHierarchyElementChangedPayload
+{
+	GENERATED_BODY()
+
+	TWeakPtr<FHierarchyElementViewModel> DeletedElementViewModel;
+};
+
+USTRUCT()
+struct FHierarchyElementChangedPayload_ElementPropertyChanged : public FHierarchyElementChangedPayload
+{
+	GENERATED_BODY()
+
+	TWeakPtr<FHierarchyElementViewModel> ElementViewModel;
+};
+
 /** The main controller class for the SHierarchyEditor widget. Defines core hierarchy rules.
  *  Inherit from this and override the required virtual functions, instantiate an object, Initialize it and pass it to the SHierarchyEditor widget. */
 UCLASS(MinimalAPI, Abstract, Transient)
@@ -443,6 +511,7 @@ public:
 	DECLARE_DELEGATE_OneParam(FOnRefreshViewRequested, bool bForceFullRefresh)
 	DECLARE_DELEGATE_OneParam(FOnNavigateToElementIdentityInHierarchyRequested, FHierarchyElementIdentity Identity)
 	DECLARE_DELEGATE_OneParam(FOnNavigateToElementInHierarchyRequested, TSharedPtr<FHierarchyElementViewModel> ElementViewModel)
+	DECLARE_DELEGATE_OneParam(FOnHierarchyElementChanged, TInstancedStruct<FHierarchyElementChangedPayload>)
 
 	GENERATED_BODY()
 
@@ -520,9 +589,8 @@ public:
 	/** Set the refresh context to easily allow Hierarchy Elements representing external data to access whether the external data still exists. */
 	void SetRefreshContext(UHierarchyDataRefreshContext* InContext) { RefreshContext = InContext; }
 	
-	UE_API UHierarchyElement* AddElementUnderRoot(TSubclassOf<UHierarchyElement> NewChildClass, FHierarchyElementIdentity ChildIdentity);
+	UE_API void AddElementUnderRoot(TSubclassOf<UHierarchyElement> NewChildClass, FHierarchyElementIdentity ChildIdentity);
 	UE_API void AddCategory(TSharedPtr<FHierarchyElementViewModel> CategoryParent = nullptr) const;
-	UE_API void AddSection() const;
 
 	/** Delete all specified elements. */
 	UE_API void DeleteElements(TArray<TSharedPtr<FHierarchyElementViewModel>> ViewModels) const;
@@ -554,18 +622,19 @@ public:
 	FSimpleDelegate& OnRefreshSectionsView() { return RefreshSectionsViewDelegate; }
 
 	// Delegates for external systems
+	FOnHierarchyElementChanged& OnHierarchyElementChanged() { return OnHierarchyElementChangedDelegate; }
 	FOnHierarchyChanged& OnHierarchyChanged() { return OnHierarchyChangedDelegate; } 
 	FOnHierarchyChanged& OnHierarchyPropertiesChanged() { return OnHierarchyPropertiesChangedDelegate; } 
-	FOnElementAdded& OnElementAdded() { return OnElementAddedDelegate; }
+	//FOnElementAdded& OnElementAdded() { return OnElementAddedDelegate; }
 	FOnRefreshViewRequested& OnRefreshViewRequested() { return RefreshAllViewsRequestedDelegate; }
 	FOnNavigateToElementIdentityInHierarchyRequested& OnNavigateToElementIdentityInHierarchyRequested() { return OnNavigateToElementIdentityInHierarchyRequestedDelegate; }
 	FOnNavigateToElementInHierarchyRequested& OnNavigateToElementInHierarchyRequested() { return OnNavigateToElementInHierarchyRequestedDelegate; }
 	FSimpleDelegate& OnInitialized() { return OnInitializedDelegate; }
-	
+
 	// Sections
 	UE_API void SetActiveHierarchySection(TSharedPtr<struct FHierarchySectionViewModel>);
 	UE_API TSharedPtr<FHierarchySectionViewModel> GetActiveHierarchySectionViewModel() const;
-	UE_API UHierarchySection* GetActiveHierarchySectionData() const;
+	UE_API const UHierarchySection* GetActiveHierarchySectionData() const;
 	UE_API bool IsHierarchySectionActive(const UHierarchySection* Section) const;
 	FOnSectionActivated& OnHierarchySectionActivated() { return OnHierarchySectionActivatedDelegate; }
 
@@ -577,7 +646,6 @@ protected:
 	UE_API virtual bool MatchesContext(const FTransactionContext& InContext, const TArray<TPair<UObject*, FTransactionObjectEvent>>& TransactionObjectContexts) const override;
 private:
 	UE_API bool FilterForHierarchySection(TSharedPtr<const FHierarchyElementViewModel> ViewModel) const;
-	UE_API bool FilterForUncategorizedRootItemsInAllSection(TSharedPtr<const FHierarchyElementViewModel> ViewModel) const;
 
 	static UE_API void GenerateDynamicContextMenu(UToolMenu* ToolMenu);
 
@@ -589,6 +657,7 @@ private:
 
 	UE_API void ToolMenuNavigateTo(const FToolMenuContext& Context) const;
 	UE_API bool ToolMenuCanNavigateTo(const FToolMenuContext& Context) const;
+
 protected:
 	UPROPERTY()
 	TObjectPtr<UHierarchyRoot> HierarchyRoot;
@@ -613,10 +682,12 @@ protected:
 	FOnNavigateToElementIdentityInHierarchyRequested OnNavigateToElementIdentityInHierarchyRequestedDelegate;
 	FOnNavigateToElementInHierarchyRequested OnNavigateToElementInHierarchyRequestedDelegate;
 	
-	FOnElementAdded OnElementAddedDelegate;
+	//FOnElementAdded OnElementAddedDelegate;
 	FOnSectionActivated OnHierarchySectionActivatedDelegate;
 	FOnSectionActivated OnSourceSectionActivatedDelegate;
 	FOnHierarchyChanged OnHierarchyChangedDelegate;
+	FOnHierarchyElementChanged OnHierarchyElementChangedDelegate;
+
 	FOnHierarchyPropertiesChanged OnHierarchyPropertiesChangedDelegate;
 	
 	FSimpleDelegate OnInitializedDelegate;
@@ -662,11 +733,11 @@ struct FHierarchyElementViewModel : TSharedFromThis<FHierarchyElementViewModel>,
 	
 	UE_API virtual ~FHierarchyElementViewModel() override;
 	
-	UHierarchyElement* GetDataMutable() const { return Element; }
+	UHierarchyElement* GetDataMutable() { return Element; }
 	const UHierarchyElement* GetData() const { return Element; }
 	
 	template<class T>
-	T* GetDataMutable() const { return Cast<T>(Element); }
+	T* GetDataMutable() { return Cast<T>(Element); }
 	
 	template<class T>
 	const T* GetData() const { return Cast<T>(Element); }
@@ -682,8 +753,8 @@ struct FHierarchyElementViewModel : TSharedFromThis<FHierarchyElementViewModel>,
 	UE_API void SyncViewModelsToData();
 	
 	/** Every item view model can define its own sort order for its children. By default we put categories above items. */
-	UE_API virtual void SortChildrenData() const;
-	
+	UE_API virtual void SortChildrenData();
+
 	const TArray<TSharedPtr<FHierarchyElementViewModel>>& GetChildren() const { return Children; }
 	TArray<TSharedPtr<FHierarchyElementViewModel>>& GetChildrenMutable() { return Children; }
 	UE_API const TArray<TSharedPtr<FHierarchyElementViewModel>>& GetFilteredChildren() const;
@@ -691,7 +762,7 @@ struct FHierarchyElementViewModel : TSharedFromThis<FHierarchyElementViewModel>,
 	UE_API void AddChildFilter(FOnFilterChild InFilterChild);
 	
 	template<class DataClass, class ViewModelChildClass>
-	void GetChildrenViewModelsForType(TArray<TSharedPtr<ViewModelChildClass>>& OutChildren, bool bRecursive = false);
+	void GetChildrenViewModelsForType(TArray<TSharedPtr<ViewModelChildClass>>& OutChildren, bool bRecursive = false) const;
 
 	/** Returns the hierarchy depth via number of parents above. */
 	UE_API int32 GetHierarchyDepth() const;
@@ -700,13 +771,13 @@ struct FHierarchyElementViewModel : TSharedFromThis<FHierarchyElementViewModel>,
 	UE_API TSharedRef<FHierarchyElementViewModel> DuplicateToThis(TSharedPtr<FHierarchyElementViewModel> ItemToDuplicate, int32 InsertIndex = INDEX_NONE);
 	UE_API TSharedRef<FHierarchyElementViewModel> ReparentToThis(TSharedPtr<FHierarchyElementViewModel> ItemToMove, int32 InsertIndex = INDEX_NONE);
 
-	UE_API TSharedPtr<FHierarchyElementViewModel> FindViewModelForChild(UHierarchyElement* Child, bool bSearchRecursively = false) const;
+	UE_API TSharedPtr<FHierarchyElementViewModel> FindViewModelForChild(const UHierarchyElement* Child, bool bSearchRecursively = false) const;
 	UE_API TSharedPtr<FHierarchyElementViewModel> FindViewModelForChild(FHierarchyElementIdentity ChildIdentity, bool bSearchRecursively = false) const;
-	UE_API int32 FindIndexOfChild(UHierarchyElement* Child) const;
-	UE_API int32 FindIndexOfDataChild(TSharedPtr<FHierarchyElementViewModel> Child) const;
-	UE_API int32 FindIndexOfDataChild(UHierarchyElement* Child) const;
+	UE_API int32 FindIndexOfChild(const UHierarchyElement* Child) const;
+	UE_API int32 FindIndexOfDataChild(TSharedPtr<const FHierarchyElementViewModel> Child) const;
+	UE_API int32 FindIndexOfDataChild(const UHierarchyElement* Child) const;
 
-	UE_API UHierarchyElement* AddChild(TSubclassOf<UHierarchyElement> NewChildClass, FHierarchyElementIdentity ChildIdentity);
+	UE_API TSharedPtr<FHierarchyElementViewModel> AddChild(TSubclassOf<UHierarchyElement> NewChildClass, int32 InsertIndex = INDEX_NONE, FHierarchyElementIdentity ChildIdentity = FHierarchyElementIdentity());
 	
 	/** Deleting will ask the parent to delete its child */
 	UE_API void Delete();
@@ -738,6 +809,9 @@ struct FHierarchyElementViewModel : TSharedFromThis<FHierarchyElementViewModel>,
 		}
 	};
 
+	/** Called when a property on the underlying data changes. */
+	virtual void PostEditChange(const FPropertyChangedEvent& PropertyChangedEvent, FProperty* PropertyThatChanged) {}
+	
 	/** Should return true if properties are supposed to be editable & needs to be true if typical operations should work on it (renaming, dragging, deleting etc.) */
 	virtual FResultWithUserFeedback IsEditableByUser() { return FResultWithUserFeedback(false); }
 	
@@ -772,8 +846,10 @@ struct FHierarchyElementViewModel : TSharedFromThis<FHierarchyElementViewModel>,
 	/** Returns true if deletable */
 	bool CanDelete() { return IsEditableByUser().bResult && CanDeleteInternal(); }
 
+	UE_API bool IsExpandedByDefault() const;
+	
 	/** Returns true if the given item can be dropped on the given target area. */
-	FResultWithUserFeedback CanDropOn(TSharedPtr<FHierarchyElementViewModel> DraggedItem, EItemDropZone ItemDropZone);
+	FResultWithUserFeedback CanDropOn(TSharedPtr<FHierarchyElementViewModel> DraggedElement, EItemDropZone ItemDropZone);
 
 	/** Gets executed when an item was dropped on this. */
 	void OnDroppedOn(TSharedPtr<FHierarchyElementViewModel> DroppedItem, EItemDropZone ItemDropZone);
@@ -783,9 +859,6 @@ struct FHierarchyElementViewModel : TSharedFromThis<FHierarchyElementViewModel>,
 	/** Returns true if this element can add a new element via the add menu. */
 	FResultWithUserFeedback CanContain(TSubclassOf<UHierarchyElement> HierarchyElementType);
 	
-	/** Determines the section this item belongs to. */
-	const UHierarchySection* GetSection() const { return GetSectionInternal(); }
-	
 	/** For data cleanup that represents external data, this needs to return true in order for live cleanup to work. */
 	virtual bool RepresentsExternalData() const { return false; }
 	/** This function determines whether a hierarchy item that represents that external data should be maintained during data refresh
@@ -794,10 +867,10 @@ struct FHierarchyElementViewModel : TSharedFromThis<FHierarchyElementViewModel>,
 	virtual bool DoesExternalDataStillExist(const UHierarchyDataRefreshContext* Context) const { return false; }
 
 	/** The UObject we display in the details panel when this item is selected. By default it's the hierarchy element the view model represents. */
-	virtual UObject* GetObjectForEditing() { return Element; }
+	UE_API virtual UObject* GetObjectForEditing() { return Element; }
 	/** Source items are transient, which is why we don't allow editing by default.
 	 * This is useful to override if source data points at actual data to edit. */
-	virtual bool AllowEditingInDetailsPanel() const { return bIsForHierarchy; }
+	UE_API virtual bool AllowEditingInDetailsPanel() const { return bIsForHierarchy; }
 
 	bool IsForHierarchy() const { return bIsForHierarchy; }
 
@@ -810,6 +883,7 @@ struct FHierarchyElementViewModel : TSharedFromThis<FHierarchyElementViewModel>,
 	FOnSynced& GetOnSynced() { return OnSyncedDelegate; }
 	FOnChildRequestedDeletion& OnChildRequestedDeletion() { return OnChildRequestedDeletionDelegate; }
 
+	/** Returns the parent view model of the hierarchical parent */
 	TWeakPtr<FHierarchyElementViewModel> GetParent() { return Parent; }
 
 	virtual const FTableRowStyle* GetRowStyle() const { return nullptr; }
@@ -817,6 +891,11 @@ struct FHierarchyElementViewModel : TSharedFromThis<FHierarchyElementViewModel>,
 protected:
 	/** Override this to handle post-add logic (for example when dropping an element on another element, or using the add button). */
 	virtual void PostAddFixup() { }
+	
+	/** Called after an element was dropped onto the calling element. Allows inserting specialized fix up. Not called if OnDroppedOnOverride was called.
+	 *  Note: DroppedElementViewModel is the newly created view model _after_ the drop occured. */
+	UE_API virtual void PostOnDroppedOn(TSharedPtr<FHierarchyElementViewModel> DroppedElementViewModel) { }
+
 private:
 	/** Should return true if this element can contain other elements as children. Can optionally specify a message for failure. */
 	virtual FResultWithUserFeedback CanContainInternal(TSubclassOf<UHierarchyElement> InHierarchyElementType) { return false; }
@@ -832,14 +911,11 @@ private:
 	/** Should return true if deletable. By default, we can delete items in the hierarchy, not in the source. */
 	virtual bool CanDeleteInternal() { return IsForHierarchy(); }
 
+	UE_API virtual bool IsExpandedByDefaultInternal() const;
+	
 	/** Should return true if the given drag drop operation is allowed to succeed. */
-	UE_API virtual FResultWithUserFeedback CanDropOnInternal(TSharedPtr<FHierarchyElementViewModel>, EItemDropZone ItemDropZone);
-
-	/** Can be overridden to support sections in the source list.
-	 * In the hierarchy only categories can be parented directly to the root, but using this it is possible to add items to custom sections in the source panel.
-	 * This will only work for top-level objects, i.e. anything directly under the root. */
-	virtual const UHierarchySection* GetSectionInternal() const { return nullptr; }
-private:
+	UE_API virtual FResultWithUserFeedback CanDropOnInternal(TSharedPtr<FHierarchyElementViewModel> DraggedElement, EItemDropZone ItemDropZone);
+private:	
 	/** Optionally implement this to refresh dependent data. */
 	virtual void RefreshChildrenDataInternal() {}
 	/** Optionally implement this to further customize the view model sync process.
@@ -869,10 +945,18 @@ protected:
 	
 	bool bRenamePending = false;
 	bool bIsForHierarchy = false;
+
+	struct FDefaultUserFeedback
+	{
+		static FText UneditableTarget;
+		static FText TargetCantContainType;
+		static FText AddElementFromSourceToHierarchy;
+		static FText MoveElementWithinHierarchy;
+	};
 };
 
 template <class DataClass, class ViewModelClass>
-void FHierarchyElementViewModel::GetChildrenViewModelsForType(TArray<TSharedPtr<ViewModelClass>>& OutChildren, bool bRecursive)
+void FHierarchyElementViewModel::GetChildrenViewModelsForType(TArray<TSharedPtr<ViewModelClass>>& OutChildren, bool bRecursive) const
 {
 	for(auto& Child : Children)
 	{
@@ -903,7 +987,7 @@ struct FHierarchyRootViewModel : FHierarchyElementViewModel
 	UE_API virtual void Initialize() override;
 
 	virtual FResultWithUserFeedback IsEditableByUser() override { return true; }
-	UE_API virtual FResultWithUserFeedback CanContainInternal(TSubclassOf<UHierarchyElement> InHierarchyElement) override;
+	UE_API virtual FResultWithUserFeedback CanContainInternal(TSubclassOf<UHierarchyElement> InHierarchyElementType) override;
 	UE_API virtual FResultWithUserFeedback CanDropOnInternal(TSharedPtr<FHierarchyElementViewModel> DraggedElement, EItemDropZone ItemDropZone) override;
 
 	virtual FResultWithUserFeedback CanHaveChildren() const override { return true; }
@@ -923,6 +1007,8 @@ struct FHierarchyRootViewModel : FHierarchyElementViewModel
 private:
 	UE_API void PropagateOnSynced();
 	UE_API virtual void SyncViewModelsToDataInternal() override;
+	UE_API virtual void PostOnDroppedOn(TSharedPtr<FHierarchyElementViewModel> DroppedElementViewModel) override;
+private:
 	TArray<TSharedPtr<struct FHierarchySectionViewModel>> SectionViewModels;
 
 	FOnSyncPropagated OnSyncPropagatedDelegate;
@@ -933,12 +1019,11 @@ private:
 
 struct FHierarchySectionViewModel : FHierarchyElementViewModel
 {
-	FHierarchySectionViewModel(UHierarchySection* InItem, TSharedRef<FHierarchyRootViewModel> InParent, TWeakObjectPtr<UDataHierarchyViewModelBase> InHierarchyViewModel);
-	
+	UE_API FHierarchySectionViewModel(UHierarchySection* InItem, TSharedRef<FHierarchyRootViewModel> InParent, TWeakObjectPtr<UDataHierarchyViewModelBase> InHierarchyViewModel);
 	virtual ~FHierarchySectionViewModel() override {}
 
 	/** A custom function to retrieve data; if it's the all section, we want to return nullptr as that's what it serializes as. */
-	UHierarchySection* GetSectionData() const;
+	UE_API const UHierarchySection* GetSectionData() const;
 	
 	UE_API virtual FString ToString() const override;
 	
@@ -950,11 +1035,11 @@ struct FHierarchySectionViewModel : FHierarchyElementViewModel
 	UE_API FText GetSectionTooltip() const;
 
 	void SetSectionImage(const FSlateBrush* InSectionImage) { SectionImage = InSectionImage; }
-	const FSlateBrush* GetSectionImage() const { return SectionImage; }
+	UE_API virtual const FSlateBrush* GetSectionImageBrush() const { return SectionImage ? SectionImage : FAppStyle::GetNoBrush(); }
 	
 	bool IsAllSection() const { return GetData()->GetFName() == UDataHierarchyViewModelBase::AllSectionHierarchyObjectName || GetData()->GetFName() == UDataHierarchyViewModelBase::AllSectionSourceObjectName; }
 protected:
-	virtual bool AllowEditingInDetailsPanel() const override;
+	UE_API virtual bool AllowEditingInDetailsPanel() const override;
 
 	/** Only hierarchy sections are editable */
 	UE_API virtual FResultWithUserFeedback IsEditableByUser() override;
@@ -999,7 +1084,8 @@ struct FHierarchyCategoryViewModel : FHierarchyElementViewModel
 {
 	FHierarchyCategoryViewModel(UHierarchyCategory* InCategory, TSharedRef<FHierarchyElementViewModel> InParent, TWeakObjectPtr<UDataHierarchyViewModelBase> InHierarchyViewModel) : FHierarchyElementViewModel(InCategory, InParent, InHierarchyViewModel, InParent->IsForHierarchy()) {}
 	virtual ~FHierarchyCategoryViewModel() override{}
-	
+
+	void SetCategoryName(FName InText) { GetDataMutable<UHierarchyCategory>()->SetCategoryName(InText); }
 	FText GetCategoryName() const { return GetData<UHierarchyCategory>()->GetCategoryAsText(); }
 	
 	virtual FResultWithUserFeedback IsEditableByUser() override { return FResultWithUserFeedback(true); }
@@ -1007,13 +1093,9 @@ struct FHierarchyCategoryViewModel : FHierarchyElementViewModel
 	virtual FResultWithUserFeedback CanDragInternal() override { return true; }
 	virtual bool CanRenameInternal() override { return true; }
 	virtual void RenameInternal(FName NewName) override { GetDataMutable<UHierarchyCategory>()->SetCategoryName(NewName); }
-	virtual const UHierarchySection* GetSectionInternal() const override { return Cast<UHierarchyCategory>(Element)->GetSection(); }
-
-	UE_API bool IsTopCategoryActive() const;
 
 	UE_API virtual FResultWithUserFeedback CanContainInternal(TSubclassOf<UHierarchyElement> InHierarchyElement) override;
 	UE_API virtual FResultWithUserFeedback CanDropOnInternal(TSharedPtr<FHierarchyElementViewModel>, EItemDropZone ItemDropZone) override;
-	UE_API virtual void PostAddFixup() override;
 
 	UE_API virtual const FTableRowStyle* GetRowStyle() const override;
 };
diff --git a/Engine/Source/Editor/UnrealEd/Private/EditorEngine.cpp b/Engine/Source/Editor/UnrealEd/Private/EditorEngine.cpp
--- a/Engine/Source/Editor/UnrealEd/Private/EditorEngine.cpp
+++ b/Engine/Source/Editor/UnrealEd/Private/EditorEngine.cpp
@@ -1522,7 +1522,8 @@ void UEditorEngine::LoadDefaultEditorModules()
 			TEXT("MassEntityEditor"),
 			TEXT("SVGDistanceField"),
 			TEXT("DataHierarchyEditor"),
-			TEXT("ToolWidgets")
+			TEXT("ToolWidgets"),
+			TEXT("UserAssetTagsEditor")
 		};
 
 	FScopedSlowTask ModuleSlowTask((float)UE_ARRAY_COUNT(ModuleNames));
diff --git a/Engine/Source/Editor/UnrealEd/Private/ObjectTools.cpp b/Engine/Source/Editor/UnrealEd/Private/ObjectTools.cpp
--- a/Engine/Source/Editor/UnrealEd/Private/ObjectTools.cpp
+++ b/Engine/Source/Editor/UnrealEd/Private/ObjectTools.cpp
@@ -4501,6 +4501,9 @@ namespace ObjectTools
 
 				// Remove any metadata from old package pointing to moved objects
 				OldPackage->GetMetaData().RemoveMetaDataOutsidePackage(OldPackage);
+
+				// Allow external systems to deal with metadata
+				UPackage::OnMetaDataTransferRequestedDelegate.Broadcast(OldPackage->GetMetaData(), NewPackage->GetMetaData(), OldPackage, NewPackage);
 			}
 
 			// Migrate any thumbnail from the old package to the new one
diff --git a/Engine/Source/Editor/UserAssetTagsEditor/Private/AssetEditors/AssetEditor_TaggedAssetBrowserConfiguration.h b/Engine/Source/Editor/UserAssetTagsEditor/Private/AssetEditors/AssetEditor_TaggedAssetBrowserConfiguration.h
new file mode 100644
--- /dev/null
+++ b/Engine/Source/Editor/UserAssetTagsEditor/Private/AssetEditors/AssetEditor_TaggedAssetBrowserConfiguration.h
@@ -0,0 +1,27 @@
+﻿// Copyright Epic Games, Inc. All Rights Reserved.
+
+#pragma once
+
+#include "CoreMinimal.h"
+#include "Tools/UAssetEditor.h"
+#include "AssetEditor_TaggedAssetBrowserConfiguration.generated.h"
+
+/**
+ * 
+ */
+UCLASS(Transient)
+class USERASSETTAGSEDITOR_API UAssetEditor_TaggedAssetBrowserConfiguration : public UAssetEditor
+{
+	GENERATED_BODY()
+
+public:
+	void SetObjectToEdit(UObject* InObject);
+
+protected:
+	virtual void GetObjectsToEdit(TArray<UObject*>& OutObjectsToEdit) override;
+	virtual TSharedPtr<FBaseAssetToolkit> CreateToolkit() override;
+
+private:
+	UPROPERTY(Transient)
+	TObjectPtr<UObject> ObjectToEdit;
+};
diff --git a/Engine/Source/Editor/UserAssetTagsEditor/Private/Assets/AssetDefinition_TaggedAssetBrowserConfiguration.h b/Engine/Source/Editor/UserAssetTagsEditor/Private/Assets/AssetDefinition_TaggedAssetBrowserConfiguration.h
new file mode 100644
--- /dev/null
+++ b/Engine/Source/Editor/UserAssetTagsEditor/Private/Assets/AssetDefinition_TaggedAssetBrowserConfiguration.h
@@ -0,0 +1,22 @@
+﻿// Copyright Epic Games, Inc. All Rights Reserved.
+
+#pragma once
+
+#include "CoreMinimal.h"
+#include "AssetDefinitionDefault.h"
+#include "AssetDefinition_TaggedAssetBrowserConfiguration.generated.h"
+
+/**
+ * 
+ */
+UCLASS()
+class USERASSETTAGSEDITOR_API UAssetDefinition_TaggedAssetBrowserConfiguration : public UAssetDefinitionDefault
+{
+	GENERATED_BODY()
+
+	virtual FText GetAssetDisplayName() const override;
+	virtual TSoftClassPtr<UObject> GetAssetClass() const override;
+	virtual FLinearColor GetAssetColor() const override;
+	virtual EAssetCommandResult OpenAssets(const FAssetOpenArgs& OpenArgs) const override;
+	virtual TConstArrayView<FAssetCategoryPath> GetAssetCategories() const override;
+};
