User Asset Tags/Data Hierarchy Editor Fixup
- 'Fixed' unresolved symbol issues (works fine locally and passes non unity tests, but so did previous CL)
- Fixed sections not being draggable (root didn't include sections as containable; excluded section drag & drop onto root directly)
- Added CanDropOnOverride function in Data Hierarchy Editor, not used currently
- Made some internal functions protected instead of public

diff --git a/Engine/Plugins/FX/Niagara/Source/NiagaraEditor/NiagaraEditor.Build.cs b/Engine/Plugins/FX/Niagara/Source/NiagaraEditor/NiagaraEditor.Build.cs
--- a/Engine/Plugins/FX/Niagara/Source/NiagaraEditor/NiagaraEditor.Build.cs
+++ b/Engine/Plugins/FX/Niagara/Source/NiagaraEditor/NiagaraEditor.Build.cs
@@ -54,8 +54,7 @@ public class NiagaraEditor : ModuleRules
 				"SparseVolumeTexture",
 				"PlacementMode",
 				"EngineSettings",
-				"EditorConfig",
-				"UserAssetTagsEditor"
+				"EditorConfig"
 			}
 		);
 
@@ -91,7 +90,8 @@ public class NiagaraEditor : ModuleRules
 				"ToolWidgets",
 				"Sequencer",
 				"DataHierarchyEditor", 
-				"MediaAssets"
+				"MediaAssets",
+				"UserAssetTagsEditor"
             }
         );
 
diff --git a/Engine/Source/Editor/DataHierarchyEditor/Private/DataHierarchyViewModelBase.cpp b/Engine/Source/Editor/DataHierarchyEditor/Private/DataHierarchyViewModelBase.cpp
--- a/Engine/Source/Editor/DataHierarchyEditor/Private/DataHierarchyViewModelBase.cpp
+++ b/Engine/Source/Editor/DataHierarchyEditor/Private/DataHierarchyViewModelBase.cpp
@@ -1105,6 +1105,12 @@ bool FHierarchyElementViewModel::IsExpandedByDefault() const
 
 FHierarchyElementViewModel::FResultWithUserFeedback FHierarchyElementViewModel::CanDropOn(TSharedPtr<FHierarchyElementViewModel> DraggedElement, EItemDropZone ItemDropZone)
 {
+	TOptional<FResultWithUserFeedback> CanDropOnOverrideResult = CanDropOnOverride(DraggedElement, ItemDropZone);
+	if(CanDropOnOverrideResult.IsSet())
+	{
+		return CanDropOnOverrideResult.GetValue();
+	}
+	
 	// Editability check
 	{
 		FResultWithUserFeedback Results = IsEditableByUser();
@@ -1185,12 +1191,12 @@ FHierarchyElementViewModel::FResultWithUserFeedback FHierarchyElementViewModel::
 		}
 	}
 	
-	return Results;;
+	return Results;
 }
 
 void FHierarchyElementViewModel::OnDroppedOn(TSharedPtr<FHierarchyElementViewModel> DroppedItem, EItemDropZone ItemDropZone)
 {
-	if(OnDroppedOnOverride(DroppedItem, ItemDropZone))
+	if(OnDroppedOnOverride(DroppedItem, ItemDropZone).IsSet())
 	{
 		return;
 	}
@@ -1358,8 +1364,8 @@ void FHierarchyRootViewModel::Initialize()
 
 FHierarchyElementViewModel::FResultWithUserFeedback FHierarchyRootViewModel::CanContainInternal(TSubclassOf<UHierarchyElement> InHierarchyElementType)
 {
-	// The root generally accepts items and categories, not sections
-	FResultWithUserFeedback Results(InHierarchyElementType->IsChildOf<UHierarchyItem>() || InHierarchyElementType->IsChildOf<UHierarchyCategory>());
+	// The root generally accepts items and categories as children; sections are not handled as typical children, but they are considered as such for this function.
+	FResultWithUserFeedback Results(InHierarchyElementType->IsChildOf<UHierarchyItem>() || InHierarchyElementType->IsChildOf<UHierarchyCategory>() || InHierarchyElementType->IsChildOf<UHierarchySection>());
 	return Results;
 }
 
@@ -1367,6 +1373,12 @@ FHierarchyElementViewModel::FResultWithUserFeedback FHierarchyRootViewModel::Can
 {
 	FResultWithUserFeedback Results(true);
 
+	// We don't allow sections to be dropped on the root directly. This disables the Root Drop Zone while still allowing the drop between sections
+	if(DraggedElement->GetData()->IsA<UHierarchySection>() && ItemDropZone == EItemDropZone::OntoItem)
+	{
+		return false;
+	}
+	
 	// If we are not in the All section, we only accept categories to drop directly on the root
 	if(HierarchyViewModel->GetActiveHierarchySectionViewModel()->IsAllSection() == false && !DraggedElement->GetData()->IsA<UHierarchyCategory>())
 	{
@@ -1586,7 +1598,7 @@ bool FHierarchySectionViewModel::CanDeleteInternal()
 	return IsForHierarchy() && !IsAllSection();
 }
 
-bool FHierarchySectionViewModel::OnDroppedOnOverride(TSharedPtr<FHierarchyElementViewModel> DroppedElementViewModel, EItemDropZone ItemDropZone)
+TOptional<bool> FHierarchySectionViewModel::OnDroppedOnOverride(TSharedPtr<FHierarchyElementViewModel> DroppedElementViewModel, EItemDropZone ItemDropZone)
 {
 	if(DroppedElementViewModel->GetData()->IsA<UHierarchySection>())
 	{
diff --git a/Engine/Source/Editor/DataHierarchyEditor/Private/Widgets/SDataHierarchyEditor.cpp b/Engine/Source/Editor/DataHierarchyEditor/Private/Widgets/SDataHierarchyEditor.cpp
--- a/Engine/Source/Editor/DataHierarchyEditor/Private/Widgets/SDataHierarchyEditor.cpp
+++ b/Engine/Source/Editor/DataHierarchyEditor/Private/Widgets/SDataHierarchyEditor.cpp
@@ -1961,7 +1961,7 @@ FReply SDataHierarchyEditor::HandleHierarchyRootDrop(const FGeometry& Geometry,
 
 FHierarchyElementViewModel::FResultWithUserFeedback SDataHierarchyEditor::CanDropOnRoot(TSharedPtr<FHierarchyElementViewModel> DraggedItem) const
 {
-	return HierarchyViewModel->GetHierarchyRootViewModel()->CanDropOnInternal(DraggedItem, EItemDropZone::OntoItem);
+	return HierarchyViewModel->GetHierarchyRootViewModel()->CanDropOn(DraggedItem, EItemDropZone::OntoItem);
 }
 
 bool SDataHierarchyEditor::OnCanDropOnRoot(TSharedPtr<FDragDropOperation> DragDropOperation) const
diff --git a/Engine/Source/Editor/DataHierarchyEditor/Public/DataHierarchyCommonTypes.h b/Engine/Source/Editor/DataHierarchyEditor/Public/DataHierarchyCommonTypes.h
--- a/Engine/Source/Editor/DataHierarchyEditor/Public/DataHierarchyCommonTypes.h
+++ b/Engine/Source/Editor/DataHierarchyEditor/Public/DataHierarchyCommonTypes.h
@@ -92,3 +92,5 @@ struct UE_API FDataHierarchyElementMetaData_SectionAssociation : public FDataHie
 	UPROPERTY()
 	TWeakObjectPtr<const class UHierarchySection> Section;
 };
+
+#undef UE_API
diff --git a/Engine/Source/Editor/DataHierarchyEditor/Public/DataHierarchyViewModelBase.h b/Engine/Source/Editor/DataHierarchyEditor/Public/DataHierarchyViewModelBase.h
--- a/Engine/Source/Editor/DataHierarchyEditor/Public/DataHierarchyViewModelBase.h
+++ b/Engine/Source/Editor/DataHierarchyEditor/Public/DataHierarchyViewModelBase.h
@@ -853,8 +853,9 @@ struct FHierarchyElementViewModel : TSharedFromThis<FHierarchyElementViewModel>,
 
 	/** Gets executed when an item was dropped on this. */
 	void OnDroppedOn(TSharedPtr<FHierarchyElementViewModel> DroppedItem, EItemDropZone ItemDropZone);
-	/** Lets you completely customize the OnDroppedOn behavior. Needs to return true if you override this. */
-	virtual bool OnDroppedOnOverride(TSharedPtr<FHierarchyElementViewModel> DroppedItem, EItemDropZone ItemDropZone) { return false; }
+	/** Lets you completely customize the OnDroppedOn behavior. Needs to a specific value if you override this. */
+	virtual TOptional<FResultWithUserFeedback> CanDropOnOverride(TSharedPtr<FHierarchyElementViewModel> DraggedElement, EItemDropZone ItemDropZone) { return {}; }
+	virtual TOptional<bool> OnDroppedOnOverride(TSharedPtr<FHierarchyElementViewModel> DroppedElement, EItemDropZone ItemDropZone) { return {}; }
 	
 	/** Returns true if this element can add a new element via the add menu. */
 	FResultWithUserFeedback CanContain(TSubclassOf<UHierarchyElement> HierarchyElementType);
@@ -986,10 +987,6 @@ struct FHierarchyRootViewModel : FHierarchyElementViewModel
 
 	UE_API virtual void Initialize() override;
 
-	virtual FResultWithUserFeedback IsEditableByUser() override { return true; }
-	UE_API virtual FResultWithUserFeedback CanContainInternal(TSubclassOf<UHierarchyElement> InHierarchyElementType) override;
-	UE_API virtual FResultWithUserFeedback CanDropOnInternal(TSharedPtr<FHierarchyElementViewModel> DraggedElement, EItemDropZone ItemDropZone) override;
-
 	virtual FResultWithUserFeedback CanHaveChildren() const override { return true; }
 
 	UE_API TSharedPtr<struct FHierarchySectionViewModel> AddSection();
@@ -1003,7 +1000,10 @@ struct FHierarchyRootViewModel : FHierarchyElementViewModel
 	/** Delegates for when a section is added or removed */
 	FOnSingleSectionChanged& OnSectionAdded() { return OnSectionAddedDelegate; }
 	FOnSingleSectionChanged& OnSectionDeleted() { return OnSectionDeletedDelegate; }
-
+protected:
+	virtual FResultWithUserFeedback IsEditableByUser() override { return true; }
+	UE_API virtual FResultWithUserFeedback CanContainInternal(TSubclassOf<UHierarchyElement> InHierarchyElementType) override;
+	UE_API virtual FResultWithUserFeedback CanDropOnInternal(TSharedPtr<FHierarchyElementViewModel> DraggedElement, EItemDropZone ItemDropZone) override;
 private:
 	UE_API void PropagateOnSynced();
 	UE_API virtual void SyncViewModelsToDataInternal() override;
@@ -1052,9 +1052,9 @@ protected:
 	virtual void RenameInternal(FName NewName) override { GetDataMutable<UHierarchySection>()->SetSectionName(NewName); }
 
 	UE_API virtual bool CanDeleteInternal() override;
-
+	
 	/** We override the default dropped on behavior given sections don't actually contain children. */
-	UE_API virtual bool OnDroppedOnOverride(TSharedPtr<FHierarchyElementViewModel> DroppedItem, EItemDropZone ItemDropZone) override;
+	UE_API virtual TOptional<bool> OnDroppedOnOverride(TSharedPtr<FHierarchyElementViewModel> DroppedItem, EItemDropZone ItemDropZone) override;
 	UE_API virtual FResultWithUserFeedback CanContainInternal(TSubclassOf<UHierarchyElement> InHierarchyElement) override;
 	UE_API virtual FResultWithUserFeedback CanDropOnInternal(TSharedPtr<FHierarchyElementViewModel> DraggedItem, EItemDropZone ItemDropZone) override;
 
diff --git a/Engine/Source/Editor/UserAssetTagsEditor/Private/DataHierarchyEditor/TaggedAssetBrowserConfigurationHierarchyViewModel.cpp b/Engine/Source/Editor/UserAssetTagsEditor/Private/DataHierarchyEditor/TaggedAssetBrowserConfigurationHierarchyViewModel.cpp
--- a/Engine/Source/Editor/UserAssetTagsEditor/Private/DataHierarchyEditor/TaggedAssetBrowserConfigurationHierarchyViewModel.cpp
+++ b/Engine/Source/Editor/UserAssetTagsEditor/Private/DataHierarchyEditor/TaggedAssetBrowserConfigurationHierarchyViewModel.cpp
@@ -152,6 +152,12 @@ namespace UE::UserAssetTags::ViewModels
 	FHierarchyElementViewModel::FResultWithUserFeedback FFilterViewModel_Root::CanDropOnInternal(TSharedPtr<FHierarchyElementViewModel> DraggedElement, EItemDropZone ItemDropZone)
 	{
 		FResultWithUserFeedback Results(true);
+
+		// We don't allow sections to be dropped on the root directly
+		if(DraggedElement->GetData()->IsA<UHierarchySection>() && ItemDropZone == EItemDropZone::OntoItem)
+		{
+			return false;
+		}
 		
 		FDataHierarchyElementMetaData_SectionAssociation SectionAssociation = DraggedElement->GetData()->FindMetaDataOfTypeOrDefault<FDataHierarchyElementMetaData_SectionAssociation>();
 		if(SectionAssociation.Section != HierarchyViewModel->GetActiveHierarchySectionViewModel()->GetSectionData())
@@ -160,7 +166,7 @@ namespace UE::UserAssetTags::ViewModels
 			Results.UserFeedback = FText::FormatOrdered(BaseMessage, DraggedElement->ToStringAsText(), HierarchyViewModel->GetActiveHierarchySectionViewModel()->GetSectionNameAsText());	
 		}
 		
-		return true;
+		return Results;
 	}
 
 	UObject* FFilterViewModel_Root::GetObjectForEditing()
diff --git a/Engine/Source/Editor/UserAssetTagsEditor/Public/Widgets/STaggedAssetBrowser.h b/Engine/Source/Editor/UserAssetTagsEditor/Public/Widgets/STaggedAssetBrowser.h
--- a/Engine/Source/Editor/UserAssetTagsEditor/Public/Widgets/STaggedAssetBrowser.h
+++ b/Engine/Source/Editor/UserAssetTagsEditor/Public/Widgets/STaggedAssetBrowser.h
@@ -18,7 +18,7 @@ class UTaggedAssetBrowserConfiguration;
 DECLARE_DELEGATE_RetVal_OneParam(TSharedRef<SWidget>, FOnGenerateAssetWidget, const FAssetData& Asset)
 
 /** You can provide custom tabs to be displayed in the right column. Useful if you want to add custom content. */
-struct FTaggedAssetBrowserCustomTabInfo
+struct USERASSETTAGSEDITOR_API FTaggedAssetBrowserCustomTabInfo
 {
 	/** The title displayed in the browser. */
 	FText Title;
@@ -33,7 +33,7 @@ struct FTaggedAssetBrowserCustomTabInfo
 	FOnAssetSelected OnAssetSelected;
 };
 
-struct FDefaultDetailsTabConfiguration
+struct USERASSETTAGSEDITOR_API FDefaultDetailsTabConfiguration
 {
 	bool bUseDefaultDetailsTab = true;
 	TOptional<FText> EmptySelectionMessage;
@@ -93,7 +93,6 @@ public:
 	SLATE_END_ARGS()
 
 	void Construct(const FArguments& InArgs, const UTaggedAssetBrowserConfiguration& InTaggedAssetBrowserConfiguration);
-	
 	virtual ~STaggedAssetBrowser() override;
 
 	void RefreshBackendFilter() const;
@@ -270,9 +269,10 @@ public:
 		SLATE_ARGUMENT(STaggedAssetBrowser::FArguments, AssetBrowserArgs)
 		SLATE_ARGUMENT(SWindow::FArguments, WindowArgs)
 	SLATE_END_ARGS()
-
+	
 	void Construct(const FArguments& InArgs, const UTaggedAssetBrowserConfiguration& Configuration);
-
+	virtual ~STaggedAssetBrowserWindow() override {}
+	
 	bool HasSelectedAssets() const;
 	TArray<FAssetData> GetSelectedAssets() const;
 	
diff --git a/Engine/Source/Editor/UserAssetTagsEditor/UserAssetTagsEditor.Build.cs b/Engine/Source/Editor/UserAssetTagsEditor/UserAssetTagsEditor.Build.cs
--- a/Engine/Source/Editor/UserAssetTagsEditor/UserAssetTagsEditor.Build.cs
+++ b/Engine/Source/Editor/UserAssetTagsEditor/UserAssetTagsEditor.Build.cs
@@ -52,8 +52,5 @@ public class UserAssetTagsEditor : ModuleRules
 				// ... add any modules that your module loads dynamically here ...
 			}
 			);
-
-		PublicDefinitions.Add("UAT_USEOLDCBMENUS=0");
-		bAllowUETypesInNamespaces = true;
 	}
 }
