[Backout] - CL45271379

diff --git a/Engine/Plugins/FX/Niagara/Source/NiagaraEditor/Public/NiagaraEmitterFactoryNew.h b/Engine/Plugins/FX/Niagara/Source/NiagaraEditor/Public/NiagaraEmitterFactoryNew.h
--- a/Engine/Plugins/FX/Niagara/Source/NiagaraEditor/Public/NiagaraEmitterFactoryNew.h
+++ b/Engine/Plugins/FX/Niagara/Source/NiagaraEditor/Public/NiagaraEmitterFactoryNew.h
@@ -1,20 +1,19 @@
 // Copyright Epic Games, Inc. All Rights Reserved.
 
 #pragma once
-
-#include "Factories/Factory.h"
 #include "NiagaraSettings.h"
-#include "FrontendFilterBase.h"
+#include "Factories/Factory.h"
 
 #include "NiagaraEmitterFactoryNew.generated.h"
 
+
 /** A factory for niagara emitter assets. */
 UCLASS(hidecategories = Object, MinimalAPI)
 class UNiagaraEmitterFactoryNew : public UFactory
 {
 	GENERATED_UCLASS_BODY()
 
-	TWeakObjectPtr<UNiagaraEmitter> EmitterToCopy;
+	UNiagaraEmitter* EmitterToCopy;
 	bool bUseInheritance;
 	bool bAddDefaultModulesAndRenderersToEmptyEmitter;
 
@@ -25,12 +24,6 @@ class UNiagaraEmitterFactoryNew : public UFactory
 
 public:
 	NIAGARAEDITOR_API static void InitializeEmitter(UNiagaraEmitter* NewEmitter, bool bAddDefaultModulesAndRenderers);
-
-private:
-	bool OnAdditionalShouldFilterAsset(const FAssetData& AssetData) const;
-	void OnExtendAddFilterMenu(UToolMenu* ToolMenu) const;
-	TArray<TSharedRef<FFrontendFilter>> OnGetExtraFrontendFilters() const;
-
 };
 
 
diff --git a/Engine/Plugins/FX/Niagara/Source/NiagaraEditor/Public/NiagaraSystemFactoryNew.h b/Engine/Plugins/FX/Niagara/Source/NiagaraEditor/Public/NiagaraSystemFactoryNew.h
--- a/Engine/Plugins/FX/Niagara/Source/NiagaraEditor/Public/NiagaraSystemFactoryNew.h
+++ b/Engine/Plugins/FX/Niagara/Source/NiagaraEditor/Public/NiagaraSystemFactoryNew.h
@@ -3,10 +3,10 @@
 #pragma once
 
 #include "CoreMinimal.h"
+
 #include "NiagaraTypes.h"
 #include "UObject/ObjectMacros.h"
 #include "Factories/Factory.h"
-#include "FrontendFilterBase.h"
 #include "NiagaraSystemFactoryNew.generated.h"
 
 class UNiagaraSystem;
@@ -30,10 +30,6 @@ public:
 
 private:
 	void TryAssignDefaultEffectType(UNiagaraSystem* System);
-	
-	bool OnAdditionalShouldFilterAsset(const FAssetData& AssetData) const;
-	void OnExtendAddFilterMenu(UToolMenu* ToolMenu) const;
-	TArray<TSharedRef<FFrontendFilter>> OnGetExtraFrontendFilters() const;
 };
 
 
diff --git a/Engine/Plugins/FX/Niagara/Source/NiagaraEditor/Public/NiagaraValidationRules.h b/Engine/Plugins/FX/Niagara/Source/NiagaraEditor/Public/NiagaraValidationRules.h
--- a/Engine/Plugins/FX/Niagara/Source/NiagaraEditor/Public/NiagaraValidationRules.h
+++ b/Engine/Plugins/FX/Niagara/Source/NiagaraEditor/Public/NiagaraValidationRules.h
@@ -373,8 +373,8 @@ public:
 
 /** This validation rule will check if a system uses emitters that are tagged as Deprecated using the Niagara Asset Tags.
  *  This is distinct from a Niagara Emitter version that is marked as deprecated, but might have a new, non-deprecated version. */
-UCLASS(Category = "Validation", DisplayName = "Check for Deprecated Emitters", meta=(DeprecationMessage="This validation rule is no longer valid. Delete it from your setup"), Deprecated)
-class UDEPRECATED_NiagaraValidationRule_CheckDeprecatedEmitters : public UNiagaraValidationRule
+UCLASS(Category = "Validation", DisplayName = "Check for Deprecated Emitters")
+class UNiagaraValidationRule_CheckDeprecatedEmitters : public UNiagaraValidationRule
 {
 	GENERATED_BODY()
 public:
diff --git a/Engine/Plugins/FX/Niagara/Source/NiagaraEditor/Public/ViewModels/HierarchyEditor/NiagaraScriptParametersHierarchyViewModel.h b/Engine/Plugins/FX/Niagara/Source/NiagaraEditor/Public/ViewModels/HierarchyEditor/NiagaraScriptParametersHierarchyViewModel.h
--- a/Engine/Plugins/FX/Niagara/Source/NiagaraEditor/Public/ViewModels/HierarchyEditor/NiagaraScriptParametersHierarchyViewModel.h
+++ b/Engine/Plugins/FX/Niagara/Source/NiagaraEditor/Public/ViewModels/HierarchyEditor/NiagaraScriptParametersHierarchyViewModel.h
@@ -99,7 +99,7 @@ struct FNiagaraHierarchyScriptCategoryViewModel : public FHierarchyCategoryViewM
 
 protected:
 	// We don't sort anything to allow free authoring of UI
-	virtual void SortChildrenData() override {}
+	virtual void SortChildrenData() const override {}
 	virtual FResultWithUserFeedback CanDropOnInternal(TSharedPtr<FHierarchyElementViewModel> DraggedElement, EItemDropZone ItemDropZone) override;
 };
 
@@ -129,7 +129,7 @@ struct FNiagaraHierarchyScriptRootViewModel : public FHierarchyRootViewModel
 
 protected:
 	// We don't sort anything to allow free authoring of UI
-	virtual void SortChildrenData() override { }
+	virtual void SortChildrenData() const override { }
 };
 
 #undef UE_API
diff --git a/Engine/Plugins/FX/Niagara/Source/NiagaraEditor/Public/ViewModels/HierarchyEditor/NiagaraSummaryViewViewModel.h b/Engine/Plugins/FX/Niagara/Source/NiagaraEditor/Public/ViewModels/HierarchyEditor/NiagaraSummaryViewViewModel.h
--- a/Engine/Plugins/FX/Niagara/Source/NiagaraEditor/Public/ViewModels/HierarchyEditor/NiagaraSummaryViewViewModel.h
+++ b/Engine/Plugins/FX/Niagara/Source/NiagaraEditor/Public/ViewModels/HierarchyEditor/NiagaraSummaryViewViewModel.h
@@ -192,11 +192,13 @@ struct FNiagaraFunctionViewModel : public FHierarchyItemViewModel
 	
 	bool IsFromBaseEmitter() const;
 
+	void SetSection(UHierarchySection& InSection) { Section = &InSection; }
+
 	bool IsDynamicInput() const { return bIsDynamicInput; }
 private:	
 	virtual void Initialize() override;
 	virtual void RefreshChildrenDataInternal() override;
-	void RefreshChildrenInputs(bool bClearCache = false);
+	void RefreshChildrenInputs(bool bClearCache = false) const;
 	
 	virtual FString ToString() const override;
 
@@ -204,9 +206,8 @@ private:
 	virtual FResultWithUserFeedback CanHaveChildren() const override { return bIsForHierarchy == false; }
 	virtual bool CanRenameInternal() override { return false; }
 	virtual FResultWithUserFeedback CanDropOnInternal(TSharedPtr<FHierarchyElementViewModel>, EItemDropZone ItemDropZone) override;
-
-	virtual bool IsExpandedByDefaultInternal() const override { return false; }
-
+	virtual const UHierarchySection* GetSectionInternal() const override;
+	
 	virtual bool RepresentsExternalData() const override { return true; }
 	virtual bool DoesExternalDataStillExist(const UHierarchyDataRefreshContext* Context) const override { ClearCache(); return GetFunctionCallNode().IsValid(); }
 	
@@ -216,6 +217,7 @@ private:
 private:
 	FDelegateHandle OnScriptAppliedHandle;
 	mutable TOptional<bool> IsFromBaseEmitterCache;
+	TWeakObjectPtr<UHierarchySection> Section;
 	bool bIsDynamicInput = false;
 };
 
@@ -249,8 +251,7 @@ protected:
 	void ClearCache() const;
 
 	void RefreshChildDynamicInputs(bool bClearCache = false);
-
-	virtual bool IsExpandedByDefaultInternal() const override { return false; }
+	
 	virtual FResultWithUserFeedback CanHaveChildren() const override;
 	virtual FResultWithUserFeedback CanContainInternal(TSubclassOf<UHierarchyElement> InHierarchyElementType) override;
 	virtual FResultWithUserFeedback IsEditableByUser() override;
@@ -291,7 +292,6 @@ struct FNiagaraAssignmentInputViewModel : public FHierarchyItemViewModel
 
 	void ClearCache() const;
 protected:
-	virtual bool IsExpandedByDefaultInternal() const override { return false; }
 	virtual FResultWithUserFeedback IsEditableByUser() override;
 	virtual bool RepresentsExternalData() const override { return true; }
 	virtual bool DoesExternalDataStillExist(const UHierarchyDataRefreshContext* Context) const override { ClearCache(); return GetInputData().IsSet(); }
@@ -311,7 +311,6 @@ struct FNiagaraHierarchySummaryCategoryViewModel : public FHierarchyCategoryView
 	bool IsFromBaseEmitter() const;
 	
 protected:
-	virtual bool IsExpandedByDefaultInternal() const override { return false; }
 	virtual FResultWithUserFeedback IsEditableByUser() override;
 	
 	mutable TOptional<bool> IsFromBaseEmitterCache;
@@ -344,18 +343,21 @@ struct FNiagaraHierarchyRendererViewModel : public FHierarchyItemViewModel
 
 	UNiagaraRendererProperties* GetRendererProperties() const;
 
+	void SetSection(UHierarchySection& InSection) { Section = &InSection; }
+
 	bool IsFromBaseEmitter() const;
 protected:
 	virtual void RefreshChildrenDataInternal() override;
-
-	virtual bool IsExpandedByDefaultInternal() const override { return false; }
 	
 	virtual FResultWithUserFeedback IsEditableByUser() override;
 	virtual FResultWithUserFeedback CanHaveChildren() const override { return bIsForHierarchy == false; }
 	
 	virtual bool RepresentsExternalData() const override { return true; }
 	virtual bool DoesExternalDataStillExist(const UHierarchyDataRefreshContext* Context) const override { return GetRendererProperties() != nullptr;}
+
+	virtual const UHierarchySection* GetSectionInternal() const override;
 	
+	TWeakObjectPtr<UHierarchySection> Section;
 	mutable TOptional<bool> IsFromBaseEmitterCache;
 };
 
@@ -368,16 +370,19 @@ struct FNiagaraHierarchyEmitterPropertiesViewModel : public FHierarchyItemViewMo
 	virtual FString ToString() const override;
 	virtual bool CanRenameInternal() override { return false; }
 
-	bool IsFromBaseEmitter() const;
-protected:
-	virtual bool IsExpandedByDefaultInternal() const override { return false; }
+	void SetSection(UHierarchySection& InSection) { Section = &InSection; }
 
+	bool IsFromBaseEmitter() const;
+protected:	
 	virtual FResultWithUserFeedback IsEditableByUser() override;
 	virtual FResultWithUserFeedback CanHaveChildren() const override { return bIsForHierarchy == false; }
 	
 	virtual bool RepresentsExternalData() const override { return true; }
 	virtual bool DoesExternalDataStillExist(const UHierarchyDataRefreshContext* Context) const override { return true; }
+
+	virtual const UHierarchySection* GetSectionInternal() const override;
 	
+	TWeakObjectPtr<UHierarchySection> Section;
 	mutable TOptional<bool> IsFromBaseEmitterCache;
 };
 
@@ -391,6 +396,8 @@ struct FNiagaraHierarchyEventHandlerViewModel : public FHierarchyItemViewModel
 
 	FNiagaraEventScriptProperties* GetEventScriptProperties() const;
 
+	void SetSection(UHierarchySection& InSection) { Section = &InSection; }
+
 	bool IsFromBaseEmitter() const;
 protected:
 	virtual void RefreshChildrenDataInternal() override;
@@ -400,7 +407,10 @@ protected:
 	
 	virtual bool RepresentsExternalData() const override { return true; }
 	virtual bool DoesExternalDataStillExist(const UHierarchyDataRefreshContext* Context) const override { return GetEventScriptProperties() != nullptr;}
+
+	virtual const UHierarchySection* GetSectionInternal() const override;
 	
+	TWeakObjectPtr<UHierarchySection> Section;
 	mutable TOptional<bool> IsFromBaseEmitterCache;
 };
 
@@ -436,7 +446,9 @@ struct FNiagaraHierarchySimStageViewModel : public FHierarchyItemViewModel
 	virtual bool CanRenameInternal() override { return false; }
 
 	UNiagaraSimulationStageBase* GetSimStage() const;
-	
+
+	void SetSection(UHierarchySection& InSection) { Section = &InSection; }
+
 	bool IsFromBaseEmitter() const;
 protected:
 	virtual void RefreshChildrenDataInternal() override;
@@ -446,7 +458,10 @@ protected:
 	
 	virtual bool RepresentsExternalData() const override { return true; }
 	virtual bool DoesExternalDataStillExist(const UHierarchyDataRefreshContext* Context) const override { return GetSimStage() != nullptr;}
+
+	virtual const UHierarchySection* GetSectionInternal() const override;
 	
+	TWeakObjectPtr<UHierarchySection> Section;
 	mutable TOptional<bool> IsFromBaseEmitterCache;
 };
 
diff --git a/Engine/Plugins/FX/Niagara/Source/NiagaraEditor/Public/ViewModels/Stack/NiagaraStackScriptHierarchyRoot.h b/Engine/Plugins/FX/Niagara/Source/NiagaraEditor/Public/ViewModels/Stack/NiagaraStackScriptHierarchyRoot.h
--- a/Engine/Plugins/FX/Niagara/Source/NiagaraEditor/Public/ViewModels/Stack/NiagaraStackScriptHierarchyRoot.h
+++ b/Engine/Plugins/FX/Niagara/Source/NiagaraEditor/Public/ViewModels/Stack/NiagaraStackScriptHierarchyRoot.h
@@ -60,9 +60,9 @@ public:
 	UNiagaraNodeFunctionCall& GetOwningFunctionCallNode() const { return *OwningFunctionCallNode.Get(); }
 	UHierarchyRoot* GetScriptParameterHierarchyRoot() const;
 	
-	NIAGARAEDITOR_API TConstArrayView<const TObjectPtr<UHierarchySection>> GerHierarchySectionData() const;
-	NIAGARAEDITOR_API const UHierarchySection* GetActiveHierarchySection() const;
-	NIAGARAEDITOR_API void SetActiveHierarchySection(const UHierarchySection* InActiveSection);
+	NIAGARAEDITOR_API const TArray<UHierarchySection*>& GetSections() const;
+	NIAGARAEDITOR_API const UHierarchySection* GetActiveSection() const;
+	NIAGARAEDITOR_API void SetActiveSection(const UHierarchySection* InActiveSection);
 
 	void SetShouldDisplayLabel(bool bInShouldDisplayLabel);
 	NIAGARAEDITOR_API bool GetShouldDisplayLabel() const;
diff --git a/Engine/Plugins/FX/Niagara/Source/NiagaraEditor/Public/ViewModels/Stack/NiagaraStackSummaryViewInputCollection.h b/Engine/Plugins/FX/Niagara/Source/NiagaraEditor/Public/ViewModels/Stack/NiagaraStackSummaryViewInputCollection.h
--- a/Engine/Plugins/FX/Niagara/Source/NiagaraEditor/Public/ViewModels/Stack/NiagaraStackSummaryViewInputCollection.h
+++ b/Engine/Plugins/FX/Niagara/Source/NiagaraEditor/Public/ViewModels/Stack/NiagaraStackSummaryViewInputCollection.h
@@ -30,7 +30,7 @@ protected:
 	NIAGARAEDITOR_API virtual void RefreshChildrenInternal(const TArray<UNiagaraStackEntry*>& CurrentChildren, TArray<UNiagaraStackEntry*>& NewChildren, TArray<FStackIssue>& NewIssues) override;
 
 	NIAGARAEDITOR_API virtual void GetSectionsInternal(TArray<FNiagaraStackSection>& OutStackSections) const override;
-	NIAGARAEDITOR_API TConstArrayView<const TObjectPtr<UHierarchySection>> GetHierarchySections() const;
+	NIAGARAEDITOR_API const TArray<UHierarchySection*>& GetHierarchySections() const;
 
 	NIAGARAEDITOR_API virtual bool FilterByActiveSection(const UNiagaraStackEntry& Child) const override;
 	bool FilterOnlyModified(const UNiagaraStackEntry& Child) const;
diff --git a/Engine/Plugins/FX/Niagara/Source/NiagaraEditor/Public/Widgets/AssetBrowser/NiagaraAssetBrowserConfig.h b/Engine/Plugins/FX/Niagara/Source/NiagaraEditor/Public/Widgets/AssetBrowser/NiagaraAssetBrowserConfig.h
new file mode 100644
--- /dev/null
+++ b/Engine/Plugins/FX/Niagara/Source/NiagaraEditor/Public/Widgets/AssetBrowser/NiagaraAssetBrowserConfig.h
@@ -0,0 +1,50 @@
+﻿// Copyright Epic Games, Inc. All Rights Reserved.
+
+#pragma once
+
+#include "CoreMinimal.h"
+#include "EditorConfigBase.h"
+#include "NiagaraAssetBrowserConfig.generated.h"
+
+#define UE_API NIAGARAEDITOR_API
+
+USTRUCT()
+struct FNiagaraAssetBrowserConfiguration
+{
+	GENERATED_BODY()
+	
+	UPROPERTY()
+	TArray<FName> MainFilterSelection;
+
+	UPROPERTY()
+	bool bShouldDisplayViewport = false;
+};
+
+UCLASS(MinimalAPI, EditorConfig="NiagaraAssetBrowser")
+class UNiagaraAssetBrowserConfig : public UEditorConfigBase
+{
+	GENERATED_BODY()
+public:
+	DECLARE_MULTICAST_DELEGATE_OneParam(FOnPropertyChanged, const FPropertyChangedEvent&);
+	
+	static UE_API void Initialize();
+	static UNiagaraAssetBrowserConfig* Get() { return Instance; }
+	FOnPropertyChanged& OnPropertyChanged() { return OnPropertyChangedDelegate; }
+	
+	UPROPERTY(meta=(EditorConfig))
+	TMap<FName, FNiagaraAssetBrowserConfiguration> MainFilterSelection;
+
+	UPROPERTY(meta=(EditorConfig))
+    bool bShowHiddenAssets = false;
+    	
+	UPROPERTY(meta=(EditorConfig))
+	bool bShowDeprecatedAssets = false;
+
+	UE_API virtual void PostEditChangeProperty(FPropertyChangedEvent& PropertyChangedEvent) override;	
+private:
+	static UE_API TObjectPtr<UNiagaraAssetBrowserConfig> Instance;
+
+	FOnPropertyChanged OnPropertyChangedDelegate;
+};
+
+#undef UE_API
diff --git a/Engine/Plugins/FX/Niagara/Source/NiagaraEditor/Public/Widgets/AssetBrowser/NiagaraAssetBrowserPreview.h b/Engine/Plugins/FX/Niagara/Source/NiagaraEditor/Public/Widgets/AssetBrowser/NiagaraAssetBrowserPreview.h
new file mode 100644
--- /dev/null
+++ b/Engine/Plugins/FX/Niagara/Source/NiagaraEditor/Public/Widgets/AssetBrowser/NiagaraAssetBrowserPreview.h
@@ -0,0 +1,55 @@
+﻿// Copyright Epic Games, Inc. All Rights Reserved.
+
+#pragma once
+
+#include "CoreMinimal.h"
+#include "NiagaraComponent.h"
+#include "SEditorViewport.h"
+#include "Widgets/SCompoundWidget.h"
+
+#define UE_API NIAGARAEDITOR_API
+
+/**
+ * 
+ */
+class SNiagaraAssetBrowserPreview : public SEditorViewport, public FGCObject
+{
+public:
+	SLATE_BEGIN_ARGS(SNiagaraAssetBrowserPreview)
+		{
+		}
+	SLATE_END_ARGS()
+
+	/** Constructs this widget with InArgs */
+	UE_API void Construct(const FArguments& InArgs);
+	UE_API virtual ~SNiagaraAssetBrowserPreview() override;
+
+	UE_API void SetEmitter(UNiagaraEmitter& Emitter);
+	UE_API void SetSystem(UNiagaraSystem& System);
+	UE_API void ResetAsset();
+	
+	UNiagaraComponent* GetPreviewComponent() { return PreviewComponent; }
+private:
+	UE_API virtual void Tick(const FGeometry& AllottedGeometry, const double InCurrentTime, const float InDeltaTime) override;
+	
+	UE_API virtual void AddReferencedObjects( FReferenceCollector& Collector ) override;
+	virtual FString GetReferencerName() const override
+	{
+		return TEXT("SNiagaraAssetBrowserPreview");
+	}
+
+private:
+	TObjectPtr<UObject> AssetToDisplay;
+	TObjectPtr<UNiagaraComponent> PreviewComponent;
+	TObjectPtr<UNiagaraSystem> SystemForEmitterDisplay;
+
+	/** Preview Scene - uses advanced preview settings */
+	TSharedPtr<class FAdvancedPreviewScene> AdvancedPreviewScene;
+	TSharedPtr<class FNiagaraAssetPreviewViewportClient> AssetPreviewViewportClient;
+
+protected:
+	/** SEditorViewport interface */
+	UE_API virtual TSharedRef<FEditorViewportClient> MakeEditorViewportClient() override;
+};
+
+#undef UE_API
diff --git a/Engine/Plugins/FX/Niagara/Source/NiagaraEditor/Public/Widgets/AssetBrowser/SNiagaraAssetBrowser.h b/Engine/Plugins/FX/Niagara/Source/NiagaraEditor/Public/Widgets/AssetBrowser/SNiagaraAssetBrowser.h
new file mode 100644
--- /dev/null
+++ b/Engine/Plugins/FX/Niagara/Source/NiagaraEditor/Public/Widgets/AssetBrowser/SNiagaraAssetBrowser.h
@@ -0,0 +1,238 @@
+﻿// Copyright Epic Games, Inc. All Rights Reserved.
+
+#pragma once
+
+#include "CoreMinimal.h"
+#include "NiagaraAssetBrowserPreview.h"
+#include "SNiagaraAssetBrowserContent.h"
+#include "NiagaraMenuFilters.h"
+#include "Widgets/SCompoundWidget.h"
+#include "Widgets/SItemSelector.h"
+#include "NiagaraAssetTagDefinitions.h"
+#include "SAssetEditorViewport.h"
+#include "Subsystems/AssetEditorSubsystem.h"
+#include "MRUFavoritesList.h"
+#include "Editor/EditorEngine.h"
+
+#define UE_API NIAGARAEDITOR_API
+
+class FNiagaraSystemViewModel;
+
+class SNiagaraAssetBrowser : public SCompoundWidget
+{
+public:
+	SLATE_BEGIN_ARGS(SNiagaraAssetBrowser)
+		: _RecentAndFavoritesList(GEditor->GetEditorSubsystem<UAssetEditorSubsystem>()->GetRecentlyOpenedAssets())
+		, _SaveSettingsName("NiagaraAssetBrowser")
+		, _AssetSelectionMode(ESelectionMode::Single)
+		{
+		}
+		SLATE_ARGUMENT(TArray<UClass*>, AvailableClasses)
+		SLATE_ATTRIBUTE(const FMainMRUFavoritesList*, RecentAndFavoritesList)
+		SLATE_ARGUMENT(TOptional<FName>, SaveSettingsName)
+		SLATE_ARGUMENT(TOptional<FText>, EmptySelectionMessage)
+		SLATE_ARGUMENT(ESelectionMode::Type, AssetSelectionMode)
+		SLATE_EVENT(FOnAssetSelected, OnAssetSelected)
+		SLATE_EVENT(FOnAssetsActivated, OnAssetsActivated)
+		SLATE_ARGUMENT(TSharedPtr<SWidget>, AdditionalWidget)
+	SLATE_END_ARGS()
+
+	UE_API void Construct(const FArguments& InArgs);
+	UE_API virtual ~SNiagaraAssetBrowser() override;
+
+	UE_API void RefreshBackendFilter() const;
+
+	UE_API TArray<FAssetData> GetSelectedAssets() const;
+	
+	UE_API TArray<UClass*> GetDisplayedAssetTypes() const;
+
+	UE_API void InitContextMenu();
+private:
+	/** This function should take into account all the different widgets states that could affect the FARFilter. */
+	UE_API FARFilter GetCurrentBackendFilter() const;
+	
+	UE_API bool ShouldFilterAsset(const FAssetData& AssetData) const;
+	
+	UE_API void PopulateFiltersSlot();
+	UE_API void PopulateAssetBrowserContentSlot();
+	UE_API void PopulateAssetBrowserDetailsSlot();
+	
+	UE_API TArray<TSharedRef<FNiagaraAssetBrowserMainFilter>> GetMainFilters() const;
+	UE_API TArray<TSharedRef<FNiagaraAssetBrowserMainFilter>> GetAllFilters() const;
+
+private:
+	/** This has to be called whenever any kind of filtering changes and refreshes the actual asset view */
+	UE_API void OnFilterChanged() const;
+
+	struct FSearchItem
+	{
+		TArray<TSharedRef<FNiagaraAssetBrowserMainFilter>> Path;
+
+		TSharedPtr<FNiagaraAssetBrowserMainFilter> GetEntry() const
+		{
+			TSharedPtr<FNiagaraAssetBrowserMainFilter> Entry = nullptr;
+			
+			if(Path.Num() > 0)
+			{
+				Entry = Path[Path.Num() - 1];
+			}
+			
+			return Entry;
+		}
+
+		bool operator==(const FSearchItem& Item) const
+		{
+			return Path == Item.Path;
+		}
+	};
+	
+	UE_API void OnFilterSearchTextChanged(const FText& Text);
+	UE_API void OnSearchButtonClicked(SSearchBox::SearchDirection SearchDirection);
+	UE_API void OnFilterSearchTextCommitted(const FText& Text, ETextCommit::Type Arg);
+	UE_API TOptional<SSearchBox::FSearchResultData> GetSearchResultData() const;
+	UE_API void GenerateSearchItems(TSharedRef<FNiagaraAssetBrowserMainFilter> Root, TArray<TSharedRef<FNiagaraAssetBrowserMainFilter>> ParentChain, TArray<FSearchItem>& OutSearchItems) const;
+	UE_API void ExpandMainFilterSearchResults();
+	UE_API void SelectNextMainFilterSearchResult();
+	UE_API void SelectPreviousMainFilterSearchResult();
+
+	UE_API void OnGetChildFiltersForFilter(TSharedRef<FNiagaraAssetBrowserMainFilter> NiagaraAssetBrowserMainFilter, TArray<TSharedRef<FNiagaraAssetBrowserMainFilter>>& OutChildren) const;
+	UE_API bool OnCompareMainFiltersForEquality(const FNiagaraAssetBrowserMainFilter& NiagaraAssetBrowserMainFilter, const FNiagaraAssetBrowserMainFilter& NiagaraAssetBrowserMainFilter1) const;
+	
+	UE_API TSharedRef<ITableRow> GenerateWidgetRowForMainFilter(TSharedRef<FNiagaraAssetBrowserMainFilter> MainFilter, const TSharedRef<STableViewBase>& OwningTable) const;
+	
+	UE_API void OnAssetSelected(const FAssetData& AssetData);
+	UE_API void OnAssetsActivated(const TArray<FAssetData>& AssetData, EAssetTypeActivationMethod::Type) const;
+	UE_API TSharedPtr<SWidget> OnGetAssetContextMenu(const TArray<FAssetData>& AssetData) const;
+	UE_API TSharedRef<SToolTip> OnGetCustomAssetTooltip(FAssetData& AssetData);
+	UE_API void OnMainFilterSelected(TSharedPtr<FNiagaraAssetBrowserMainFilter> MainFilter, ESelectInfo::Type Arg);
+	
+	UE_API void OnAssetTagActivated(const FNiagaraAssetTagDefinition& NiagaraAssetTagDefinition);
+
+	UE_API TArray<TSharedRef<FFrontendFilter>> OnGetExtraFrontendFilters() const;
+	UE_API void OnExtendAddFilterMenu(UToolMenu* ToolMenu) const;
+
+	UE_API bool OnIsAssetRecent(const FAssetData& AssetCandidate) const;
+
+	bool GetShouldDisplayViewport() const { return bShouldDisplayViewport; }
+	
+	UE_API EVisibility OnGetViewportVisibility() const;
+	UE_API EVisibility OnGetThumbnailVisibility() const;
+	UE_API EVisibility OnGetShouldDisplayVisibilityCheckbox() const;
+
+	UE_API ECheckBoxState OnShouldDisplayViewport() const;
+	UE_API void OnShouldDisplayViewportChanged(ECheckBoxState CheckBoxState);
+	UE_API FText OnGetShouldDisplayViewportTooltip() const;
+
+	UE_API void OnAssetBrowserConfigPropertyChanged(const FPropertyChangedEvent& PropertyChangedEvent);
+	
+	UE_API void SaveSettings() const;
+	UE_API void LoadSettings();
+private:
+	TArray<UClass*> AvailableClasses;
+	TAttribute<const FMainMRUFavoritesList*> RecentAndFavoritesList;
+	TOptional<FName> SaveSettingsName;
+	TOptional<FText> EmptySelectionMessage;
+	FOnAssetSelected OnAssetSelectedDelegate;
+	FOnAssetsActivated OnAssetsActivatedDelegate;
+	ESelectionMode::Type AssetSelectionMode = ESelectionMode::Single;
+	bool bShouldDisplayViewport = false;
+
+	mutable TMap<FNiagaraAssetTagDefinition, TSharedRef<FFrontendFilter>> DropdownFilterCache;
+	TArray<TSharedRef<FNiagaraAssetBrowserMainFilter>> AssetBrowserMainFilters;
+
+	TArray<FSearchItem> SourceSearchResults;
+	TOptional<FSearchItem> FocusedSearchResult;
+private:
+	TSharedPtr<STreeView<TSharedRef<FNiagaraAssetBrowserMainFilter>>> MainFilterSelector;
+	TSharedPtr<SNiagaraAssetBrowserContent> AssetBrowserContent;
+	TSharedPtr<SWidget> DetailsContainer;
+	TSharedPtr<SWidgetSwitcher> DetailsSwitcher;
+	TSharedPtr<SNiagaraAssetBrowserPreview> PreviewViewport;
+
+	SSplitter::FSlot* FiltersSlot = nullptr;
+	SSplitter::FSlot* AssetBrowserContentSlot = nullptr;
+	SSplitter::FSlot* AssetBrowserDetailsAreaSlot = nullptr;
+	SWidgetSwitcher::FSlot* AssetBrowserDetailsSlot = nullptr;
+	SVerticalBox::FSlot* AdditionalWidgetSlot = nullptr;
+
+	/** Used to temporarily stop saving & loading config state.
+	 * This is used for initialization where setting the default state should not be saved into config.*/
+	bool bSuppressSaveAndLoad = false;
+	
+	/** We use this fallback as a means to save our last selected filter, in case we have no valid active selections anymore. */
+	mutable FName LastSelectedMainFilterIdentifierFallback = NAME_None;
+};
+
+class SNiagaraAssetBrowserWindow : public SWindow
+{
+public:
+	SLATE_BEGIN_ARGS(SNiagaraAssetBrowserWindow)
+		{
+		}
+		SLATE_ARGUMENT(SNiagaraAssetBrowser::FArguments, AssetBrowserArgs)
+		SLATE_ARGUMENT(TOptional<FText>, WindowTitle)
+	SLATE_END_ARGS()
+
+	UE_API void Construct(const FArguments& InArgs);
+
+	UE_API bool HasSelectedAssets() const;
+	UE_API TArray<FAssetData> GetSelectedAssets() const;
+	
+protected:	
+	UE_API void OnAssetsActivated(const TArray<FAssetData>& AssetData, EAssetTypeActivationMethod::Type ActivationType);
+
+	UE_API virtual void OnFocusLost(const FFocusEvent& InFocusEvent) override;
+	UE_API virtual FReply OnKeyDown(const FGeometry& MyGeometry, const FKeyEvent& InKeyEvent) override;
+protected:
+	TSharedPtr<SNiagaraAssetBrowser> AssetBrowser;
+private:
+	FOnAssetsActivated OnAssetsActivatedDelegate;
+};
+
+class SNiagaraCreateAssetWindow : public SNiagaraAssetBrowserWindow
+{
+	SLATE_BEGIN_ARGS(SNiagaraCreateAssetWindow)
+		{
+		}
+		SLATE_ARGUMENT(SNiagaraAssetBrowserWindow::FArguments, AssetBrowserWindowArgs)
+	SLATE_END_ARGS()
+
+	UE_API void Construct(const FArguments& InArgs, UClass& CreatedClass);
+
+	bool ShouldProceedWithAction() const { return bProceedWithAction; }
+private:
+	/** The function that will be called by our buttons or by the asset picker itself if double-clicking, hitting enter etc. */
+	UE_API void OnAssetsActivatedInternal(const TArray<FAssetData>& AssetData, EAssetTypeActivationMethod::Type);
+	
+	UE_API FReply Proceed();
+	UE_API FReply Cancel();
+
+	UE_API FText GetCreateButtonTooltip() const;
+private:
+	TWeakObjectPtr<UClass> CreatedClass;
+	bool bProceedWithAction = false;
+};
+
+class SNiagaraAddEmitterToSystemWindow : public SNiagaraAssetBrowserWindow
+{
+	SLATE_BEGIN_ARGS(SNiagaraAddEmitterToSystemWindow)
+	{
+	}
+		SLATE_ARGUMENT(SNiagaraAssetBrowserWindow::FArguments, AssetBrowserWindowArgs)
+	SLATE_END_ARGS()
+
+	UE_API void Construct(const FArguments& InArgs, TSharedRef<FNiagaraSystemViewModel> SystemViewModel);
+private:
+	/** The function that will be called by our buttons or by the asset picker itself if double-clicking, hitting enter etc. */
+	UE_API void OnAssetsActivatedInternal(const TArray<FAssetData>& AssetData, EAssetTypeActivationMethod::Type) const;
+
+	UE_API FReply AddEmptyEmitter();
+	UE_API FReply AddSelectedEmitters();
+	UE_API FReply Cancel();
+	
+	UE_API FText GetAddButtonTooltip() const;
+private:
+	TWeakPtr<FNiagaraSystemViewModel> WeakSystemViewModel;
+};
+
+#undef UE_API
diff --git a/Engine/Plugins/FX/Niagara/Source/NiagaraEditor/Public/Widgets/AssetBrowser/SNiagaraAssetBrowserContent.h b/Engine/Plugins/FX/Niagara/Source/NiagaraEditor/Public/Widgets/AssetBrowser/SNiagaraAssetBrowserContent.h
new file mode 100644
--- /dev/null
+++ b/Engine/Plugins/FX/Niagara/Source/NiagaraEditor/Public/Widgets/AssetBrowser/SNiagaraAssetBrowserContent.h
@@ -0,0 +1,41 @@
+﻿// Copyright Epic Games, Inc. All Rights Reserved.
+
+#pragma once
+
+#include "CoreMinimal.h"
+#include "FrontendFilterBase.h"
+#include "NiagaraAssetTagDefinitions.h"
+#include "Filters/GenericFilter.h"
+#include "Filters/SAssetFilterBar.h"
+#include "Widgets/SCompoundWidget.h"
+
+#define UE_API NIAGARAEDITOR_API
+
+class SNiagaraAssetBrowserContent : public SCompoundWidget
+{
+public:
+	DECLARE_DELEGATE_RetVal_OneParam(bool, FOnShouldFilterAsset, const FAssetData&)
+	DECLARE_DELEGATE_RetVal(TArray<TSharedRef<FFrontendFilter>>, FOnGetExtraFrontendFilters)
+	
+	SLATE_BEGIN_ARGS(SNiagaraAssetBrowserContent)
+		{
+		}
+		SLATE_ARGUMENT(FAssetPickerConfig, InitialConfig)
+	SLATE_END_ARGS()
+
+	/** Constructs this widget with InArgs */
+	UE_API void Construct(const FArguments& InArgs);
+
+	UE_API void SetARFilter(FARFilter InFilter);
+
+	UE_API TArray<FAssetData> GetCurrentSelection() const;
+private:	
+	FRefreshAssetViewDelegate RefreshAssetViewDelegate;
+	FSyncToAssetsDelegate SyncToAssetsDelegate;
+	FSetARFilterDelegate SetNewFilterDelegate;
+	FGetCurrentSelectionDelegate GetCurrentSelectionDelegate;
+
+	FAssetPickerConfig InitialAssetPickerConfig;
+};
+
+#undef UE_API
diff --git a/Engine/Plugins/FX/Niagara/Source/NiagaraEditor/Public/Widgets/AssetBrowser/SNiagaraSelectedAssetDetails.h b/Engine/Plugins/FX/Niagara/Source/NiagaraEditor/Public/Widgets/AssetBrowser/SNiagaraSelectedAssetDetails.h
new file mode 100644
--- /dev/null
+++ b/Engine/Plugins/FX/Niagara/Source/NiagaraEditor/Public/Widgets/AssetBrowser/SNiagaraSelectedAssetDetails.h
@@ -0,0 +1,123 @@
+﻿// Copyright Epic Games, Inc. All Rights Reserved.
+
+#pragma once
+
+#include "CoreMinimal.h"
+#include "NiagaraAssetTagDefinitions.h"
+#include "Widgets/SCompoundWidget.h"
+#include "AssetThumbnail.h"
+#include "AssetRegistry/AssetData.h"
+
+#define UE_API NIAGARAEDITOR_API
+
+struct FDisplayedPropertyData
+{
+	DECLARE_DELEGATE_RetVal_OneParam(bool, FShouldDisplayProperty, const FAssetData&)
+	DECLARE_DELEGATE_RetVal_OneParam(TSharedRef<SWidget>, FGenerateWidget, const FAssetData&)
+
+	FShouldDisplayProperty ShouldDisplayPropertyDelegate;
+	FGenerateWidget NameWidgetDelegate;
+	FGenerateWidget ValueWidgetDelegate;
+};
+
+struct FNiagaraAssetDetailClassInfo
+{
+	DECLARE_DELEGATE_RetVal_OneParam(FText, FGetDescription, const FAssetData&);
+	
+	FGetDescription GetDescriptionDelegate;
+	TArray<FDisplayedPropertyData> DisplayedProperties;
+};
+
+struct FNiagaraAssetDetailDatabase
+{
+	static void Init();
+	static TMap<UClass*, FNiagaraAssetDetailClassInfo> NiagaraAssetDetailDatabase;
+};
+
+DECLARE_DELEGATE_OneParam(FOnAssetTagActivated, const FNiagaraAssetTagDefinition& AssetTagDefinition);
+
+class SNiagaraAssetTag : public SCompoundWidget
+{
+	SLATE_BEGIN_ARGS(SNiagaraAssetTag)
+	{
+	}
+		SLATE_EVENT(FOnAssetTagActivated, OnAssetTagActivated)
+		SLATE_ARGUMENT(TOptional<FText>, OnAssetTagActivatedTooltip)
+	SLATE_END_ARGS()
+
+	UE_API void Construct(const FArguments& InArgs, const FNiagaraAssetTagDefinition& AssetTagDefinition);
+
+private:
+	UE_API FReply OnClicked() const;
+
+private:
+	FNiagaraAssetTagDefinition AssetTagDefinition;
+	FOnAssetTagActivated OnAssetTagActivated;
+	TOptional<FText> OnAssetTagActivatedTooltip;
+};
+
+class SNiagaraAssetTagRow : public SCompoundWidget
+{
+	SLATE_BEGIN_ARGS(SNiagaraAssetTagRow)
+		{
+		}
+		SLATE_EVENT(FOnAssetTagActivated, OnAssetTagActivated)
+		SLATE_ARGUMENT(TOptional<FText>, OnAssetTagActivatedTooltip)
+		SLATE_ARGUMENT(TOptional<ENiagaraAssetTagDefinitionImportance>, DisplayType)
+	SLATE_END_ARGS()
+
+	UE_API void Construct(const FArguments& InArgs, const FAssetData& Asset);
+
+private:
+	UE_API FText GetDisplayTypeTooltipText(ENiagaraAssetTagDefinitionImportance DisplayType) const;
+};
+
+/**
+ * 
+ */
+class SNiagaraSelectedAssetDetails : public SCompoundWidget
+{
+public:
+	
+	SLATE_BEGIN_ARGS(SNiagaraSelectedAssetDetails)
+		: _ShowThumbnail(EVisibility::Visible)
+		, _MaxDesiredDescriptionWidth(300.f)
+		, _MaxDesiredPropertiesWidth(300.f)
+		{
+		}
+		SLATE_ATTRIBUTE(EVisibility, ShowThumbnail)
+		SLATE_EVENT(FOnAssetTagActivated, OnAssetTagActivated)
+		SLATE_ARGUMENT(TOptional<FText>, OnAssetTagActivatedTooltip)
+	
+		/* Text will automatically wrap according to actual width.
+		 * With this you can specify the maximum width you want the description to be; this is not the actual width.
+		 * Actual width might be increased due to big asset names.
+		 */
+		SLATE_ARGUMENT(float, MaxDesiredDescriptionWidth)
+		SLATE_ARGUMENT(float, MaxDesiredPropertiesWidth)
+
+	SLATE_END_ARGS()
+	
+	/** Constructs this widget with InArgs */
+	UE_API void Construct(const FArguments& InArgs, const FAssetData& Asset);
+
+private:
+	UE_API FReply CopyAssetPathToClipboard() const;
+	
+private:
+	FAssetData AssetData;
+	TAttribute<EVisibility> ShowThumbnail;
+	TSharedPtr<FAssetThumbnail> CurrentAssetThumbnail;
+	FOnAssetTagActivated OnAssetTagActivated;
+	TOptional<FText> OnAssetTagActivatedTooltip;
+
+	UE_API TSharedRef<SWidget> CreateAssetThumbnailWidget();
+	UE_API TSharedRef<SWidget> CreateTitleWidget();
+	UE_API TSharedRef<SWidget> CreateTypeWidget();
+	UE_API TSharedRef<SWidget> CreatePathWidget();
+	UE_API TSharedRef<SWidget> CreateDescriptionWidget();
+	UE_API TSharedRef<SWidget> CreateOptionalPropertiesList();
+	UE_API TSharedRef<SWidget> CreateAssetTagRow();
+};
+
+#undef UE_API
diff --git a/Engine/Plugins/FX/Niagara/Source/NiagaraEditorWidgets/NiagaraEditorWidgets.Build.cs b/Engine/Plugins/FX/Niagara/Source/NiagaraEditorWidgets/NiagaraEditorWidgets.Build.cs
--- a/Engine/Plugins/FX/Niagara/Source/NiagaraEditorWidgets/NiagaraEditorWidgets.Build.cs
+++ b/Engine/Plugins/FX/Niagara/Source/NiagaraEditorWidgets/NiagaraEditorWidgets.Build.cs
@@ -36,8 +36,7 @@ public class NiagaraEditorWidgets : ModuleRules
 			"ToolWidgets",
 			"KismetWidgets",
 			"ToolMenus",
-			"DataHierarchyEditor",
-			"UserAssetTagsEditor"
+			"DataHierarchyEditor"
 		});
 
 		PrivateIncludePathModuleNames.AddRange(new string[] {
diff --git a/Engine/Plugins/FX/Niagara/Source/NiagaraEditorWidgets/Private/SNiagaraOverviewGraph.cpp b/Engine/Plugins/FX/Niagara/Source/NiagaraEditorWidgets/Private/SNiagaraOverviewGraph.cpp
--- a/Engine/Plugins/FX/Niagara/Source/NiagaraEditorWidgets/Private/SNiagaraOverviewGraph.cpp
+++ b/Engine/Plugins/FX/Niagara/Source/NiagaraEditorWidgets/Private/SNiagaraOverviewGraph.cpp
@@ -23,17 +23,14 @@
 #include "SNiagaraOverviewGraphTitleBar.h"
 #include "SNiagaraStack.h"
 #include "ScopedTransaction.h"
-#include "Assets/TaggedAssetBrowserConfiguration.h"
 #include "Interfaces/IMainFrameModule.h"
-#include "Logging/StructuredLog.h"
 #include "Stateless/NiagaraStatelessEmitter.h"
 #include "Templates/SharedPointer.h"
 #include "ViewModels/NiagaraEmitterHandleViewModel.h"
 #include "ViewModels/NiagaraOverviewGraphViewModel.h"
 #include "ViewModels/NiagaraSystemViewModel.h"
 #include "Widgets/SItemSelector.h"
-#include "Widgets/STaggedAssetBrowser.h"
-#include "Widgets/AssetBrowser/SNiagaraAddEmitterToSystemWindow.h"
+#include "Widgets/AssetBrowser/SNiagaraAssetBrowser.h"
 
 #define LOCTEXT_NAMESPACE "NiagaraOverviewGraph"
 
@@ -223,17 +220,8 @@ void SNiagaraOverviewGraph::OpenAddEmitterMenu()
 	IMainFrameModule& MainFrame = FModuleManager::LoadModuleChecked<IMainFrameModule>("MainFrame");
 	TSharedPtr<SWindow>	ParentWindow = MainFrame.GetParentWindow();
 
-	FSoftObjectPath Path = GetDefault<UNiagaraEditorSettings>()->EmitterAssetWizardConfiguration;
-	UObject* ConfigUObject = Path.TryLoad();
-	if(UTaggedAssetBrowserConfiguration* ConfigurationAsset = Cast<UTaggedAssetBrowserConfiguration>(ConfigUObject))
-	{
-		TSharedRef<SWindow> AddEmitterWindow = SNew(SNiagaraAddEmitterToSystemWindow, *ConfigurationAsset, ViewModel->GetSystemViewModel());
-		FSlateApplication::Get().AddModalWindow(AddEmitterWindow, ParentWindow);
-	}
-	else
-	{
-		UE_LOGFMT(LogNiagaraEditor, Warning, "Could not find emitter wizard configuration asset. Check Niagara Editor settings under Project Settings");
-	}
+	TSharedRef<SWindow> AddEmitterWindow = SNew(SNiagaraAddEmitterToSystemWindow, ViewModel->GetSystemViewModel());
+	FSlateApplication::Get().AddModalWindow(AddEmitterWindow, ParentWindow);
 }
 
 bool SNiagaraOverviewGraph::CanAddEmitters() const
diff --git a/Engine/Plugins/FX/Niagara/Source/NiagaraEditorWidgets/Private/SNiagaraOverviewGraphTitleBar.cpp b/Engine/Plugins/FX/Niagara/Source/NiagaraEditorWidgets/Private/SNiagaraOverviewGraphTitleBar.cpp
--- a/Engine/Plugins/FX/Niagara/Source/NiagaraEditorWidgets/Private/SNiagaraOverviewGraphTitleBar.cpp
+++ b/Engine/Plugins/FX/Niagara/Source/NiagaraEditorWidgets/Private/SNiagaraOverviewGraphTitleBar.cpp
@@ -8,7 +8,6 @@
 #include "ViewModels/NiagaraOverviewGraphViewModel.h"
 #include "SNiagaraScalabilityPreviewSettings.h"
 #include "SNiagaraSystemEffectTypeBar.h"
-#include "UserAssetTagEditorUtilities.h"
 #include "AssetRegistry/AssetRegistryModule.h"
 #include "ViewModels/NiagaraEmitterHandleViewModel.h"
 
@@ -181,9 +180,12 @@ FText SNiagaraOverviewGraphTitleBar::GetSystemSubheaderText() const
 			{
 				if (FVersionedNiagaraEmitterData* ParentData = EmitterData->GetParent().GetEmitterData())
 				{
-					bool bHasDeprecatedParentAsset = UE::UserAssetTags::HasUserAssetTag(EmitterData->GetParent().Emitter, "Deprecated");
+					bool bIsDeprecatedAsset = EmitterData->GetParent().Emitter->AssetTags.ContainsByPredicate([](const FNiagaraAssetTagDefinitionReference& AssetTagDefinitionReferenceCandidate)
+					{
+						return AssetTagDefinitionReferenceCandidate.GetTagDefinitionReferenceGuid() == INiagaraModule::Get().DeprecatedTagDefinition.TagGuid;
+					});
 					
-					if (bHasDeprecatedParentAsset || ParentData->bDeprecated)
+					if (bIsDeprecatedAsset || ParentData->bDeprecated)
 					{
 						DeprecatedEmitterNames.Add(FText::FromName(EmitterModel->GetName()));
 					}
@@ -300,8 +302,10 @@ bool SNiagaraOverviewGraphTitleBar::IsUsingDeprecatedEmitter() const
 			{
 				if(FVersionedNiagaraEmitterData* ParentData = EmitterData->GetParent().GetEmitterData())
 				{
-					bool bHasDeprecatedParentAsset = UE::UserAssetTags::HasUserAssetTag(EmitterData->GetParent().Emitter, "Deprecated");
-					return ParentData->bDeprecated || bHasDeprecatedParentAsset;
+					return ParentData->bDeprecated || EmitterData->GetParent().Emitter->AssetTags.ContainsByPredicate([](const FNiagaraAssetTagDefinitionReference& AssetTagDefinitionReferenceCandidate)
+					{
+						return AssetTagDefinitionReferenceCandidate.GetTagDefinitionReferenceGuid() == INiagaraModule::Get().DeprecatedTagDefinition.TagGuid;
+					});
 				}
 			}
 		}
diff --git a/Engine/Plugins/FX/Niagara/Source/NiagaraEditorWidgets/Private/Stack/SNiagaraStackScriptHierarchyRoot.cpp b/Engine/Plugins/FX/Niagara/Source/NiagaraEditorWidgets/Private/Stack/SNiagaraStackScriptHierarchyRoot.cpp
--- a/Engine/Plugins/FX/Niagara/Source/NiagaraEditorWidgets/Private/Stack/SNiagaraStackScriptHierarchyRoot.cpp
+++ b/Engine/Plugins/FX/Niagara/Source/NiagaraEditorWidgets/Private/Stack/SNiagaraStackScriptHierarchyRoot.cpp
@@ -81,21 +81,21 @@ EVisibility SNiagaraStackScriptHierarchyRoot::GetLabelVisibility() const
 
 ECheckBoxState SNiagaraStackScriptHierarchyRoot::IsSectionChecked(const UHierarchySection* NiagaraHierarchySection) const
 {
-	return ModuleHierarchyRoot->GetActiveHierarchySection() == NiagaraHierarchySection ? ECheckBoxState::Checked : ECheckBoxState::Unchecked;
+	return ModuleHierarchyRoot->GetActiveSection() == NiagaraHierarchySection ? ECheckBoxState::Checked : ECheckBoxState::Unchecked;
 }
 
 void SNiagaraStackScriptHierarchyRoot::OnCheckStateChanged(ECheckBoxState CheckBoxState, const UHierarchySection* NiagaraHierarchySection)
 {
 	if(CheckBoxState == ECheckBoxState::Checked)
 	{
-		ModuleHierarchyRoot->SetActiveHierarchySection(NiagaraHierarchySection);
+		ModuleHierarchyRoot->SetActiveSection(NiagaraHierarchySection);
 	}
 }
 
 void SNiagaraStackScriptHierarchyRoot::ConstructSectionButtons()
 {
 	SectionSelectorBox->ClearChildren();
-	for (const UHierarchySection* Section : ModuleHierarchyRoot->GerHierarchySectionData())
+	for (const UHierarchySection* Section : ModuleHierarchyRoot->GetSections())
 	{
 		SectionSelectorBox->AddSlot()
 		[
diff --git a/Engine/Plugins/FX/Niagara/Source/NiagaraEditorWidgets/Private/Stack/SNiagaraStackTableRow.cpp b/Engine/Plugins/FX/Niagara/Source/NiagaraEditorWidgets/Private/Stack/SNiagaraStackTableRow.cpp
--- a/Engine/Plugins/FX/Niagara/Source/NiagaraEditorWidgets/Private/Stack/SNiagaraStackTableRow.cpp
+++ b/Engine/Plugins/FX/Niagara/Source/NiagaraEditorWidgets/Private/Stack/SNiagaraStackTableRow.cpp
@@ -664,8 +664,7 @@ void SNiagaraStackTableRow::ToggleShowInSummaryView() const
 			{
 				UClass* HierarchyClass = DetermineHierarchyClassForSummaryView();
 				check(HierarchyClass);
-				UHierarchyElement* Element = StackEntry->GetEmitterViewModel()->GetSummaryHierarchyViewModel()->GetHierarchyRoot()->AddChild<UHierarchyElement>(HierarchyClass);
-				Element->SetIdentity(Identity);
+				StackEntry->GetEmitterViewModel()->GetSummaryHierarchyViewModel()->AddElementUnderRoot(HierarchyClass, Identity);
 			}
 		}
 	}
diff --git a/Engine/Source/Editor/DataHierarchyEditor/Internal/Widgets/SDataHierarchyEditor.h b/Engine/Source/Editor/DataHierarchyEditor/Internal/Widgets/SDataHierarchyEditor.h
--- a/Engine/Source/Editor/DataHierarchyEditor/Internal/Widgets/SDataHierarchyEditor.h
+++ b/Engine/Source/Editor/DataHierarchyEditor/Internal/Widgets/SDataHierarchyEditor.h
@@ -35,7 +35,7 @@ class SHierarchySection : public SCompoundWidget
 	SLATE_END_ARGS()
 
 public:
-	void Construct(const FArguments& InArgs, TSharedPtr<struct FHierarchySectionViewModel> InSection);
+	UE_API void Construct(const FArguments& InArgs, TSharedPtr<struct FHierarchySectionViewModel> InSection);
 	UE_API virtual ~SHierarchySection() override;
 
 	UE_API void TryEnterEditingMode() const;
@@ -65,9 +65,6 @@ private:
 
 	UE_API UHierarchySection* TryGetSectionData() const;
 
-	const FSlateBrush* GetImageBrush() const;
-	EVisibility GetImageVisibility() const;
-
 	UE_API FText GetText() const;
 	UE_API FText GetTooltipText() const;
 	UE_API void OnRenameSection(const FText& Text, ETextCommit::Type CommitType) const;
@@ -160,8 +157,11 @@ private:
 
 	UE_API void NavigateToMatchingHierarchyElementFromSelectedSourceElement() const;
 	UE_API bool CanNavigateToMatchingHierarchyElementFromSelectedSourceElement() const;
-	
-	UE_API void OnHierarchyElementChanged(TInstancedStruct<FHierarchyElementChangedPayload> Payload);
+
+	UE_API void DeleteActiveSection() const;
+	UE_API bool CanDeleteActiveSection() const;
+
+	UE_API void OnElementAdded(TSharedPtr<FHierarchyElementViewModel> AddedItem);
 	UE_API void OnHierarchySectionActivated(TSharedPtr<FHierarchySectionViewModel> Section);
 	UE_API void OnSourceSectionActivated(TSharedPtr<FHierarchySectionViewModel> Section);
 	UE_API void OnHierarchySectionAdded(TSharedPtr<FHierarchySectionViewModel> AddedSection);
@@ -212,8 +212,6 @@ private:
 	UE_API void SelectNextSourceSearchResult();
 	UE_API void SelectPreviousSourceSearchResult();
 	UE_API TOptional<SSearchBox::FSearchResultData> GetSearchResultData() const;
-
-	void ExpandEntriesByDefault(TSharedPtr<STreeView<TSharedPtr<FHierarchyElementViewModel>>> TreeView) const;
 	
 	FHierarchyElementViewModel::FResultWithUserFeedback CanDropOnRoot(TSharedPtr<FHierarchyElementViewModel> DraggedItem) const;
 
diff --git a/Engine/Source/Editor/DataHierarchyEditor/Private/DataHierarchyEditorStyle.cpp b/Engine/Source/Editor/DataHierarchyEditor/Private/DataHierarchyEditorStyle.cpp
--- a/Engine/Source/Editor/DataHierarchyEditor/Private/DataHierarchyEditorStyle.cpp
+++ b/Engine/Source/Editor/DataHierarchyEditor/Private/DataHierarchyEditorStyle.cpp
@@ -39,7 +39,7 @@ void FDataHierarchyEditorStyle::InitDataHierarchyEditor()
 	
 	Set("HierarchyEditor.Drop.Section.Left", new CORE_BOX_BRUSH("Common/DropZoneIndicator_Above", FMargin(10.0f / 16.0f, 10.0f / 16.0f, 0.f, 0.f), SelectionColor));
 	Set("HierarchyEditor.Drop.Section.Onto", new CORE_BOX_BRUSH("Common/DropZoneIndicator_Onto", FMargin(4.0f / 16.0f), SelectionColor));
-	Set("HierarchyEditor.Drop.Section.Right", new CORE_BOX_BRUSH("Common/DropZoneIndicator_Below", FMargin(10.0f / 16.0f, 0.f, 0.f, 10.0f / 16.0f), SelectionColor));
+	Set("HierarchyEditor.Drop.Section.Below", new CORE_BOX_BRUSH("Common/DropZoneIndicator_Below", FMargin(10.0f / 16.0f, 0.f, 0.f, 10.0f / 16.0f), SelectionColor));
 	
 	const FTextBlockStyle NormalText = FAppStyle::GetWidgetStyle<FTextBlockStyle>("NormalText");
 	FSlateFontInfo CategoryFont = FAppStyle::Get().GetFontStyle(TEXT("DetailsView.CategoryFontStyle"));
diff --git a/Engine/Source/Editor/DataHierarchyEditor/Private/DataHierarchyViewModelBase.cpp b/Engine/Source/Editor/DataHierarchyEditor/Private/DataHierarchyViewModelBase.cpp
--- a/Engine/Source/Editor/DataHierarchyEditor/Private/DataHierarchyViewModelBase.cpp
+++ b/Engine/Source/Editor/DataHierarchyEditor/Private/DataHierarchyViewModelBase.cpp
@@ -18,7 +18,6 @@
 #include "Framework/Notifications/NotificationManager.h"
 #include "ToolMenus.h"
 #include "Logging/StructuredLog.h"
-#include "DataHierarchyElementCustomVersion.h"
 #include "Widgets/Notifications/SNotificationList.h"
 
 #include UE_INLINE_GENERATED_CPP_BY_NAME(DataHierarchyViewModelBase)
@@ -29,11 +28,6 @@ FName UDataHierarchyViewModelBase::AllSectionHierarchyObjectName = "AllSection_H
 FName UDataHierarchyViewModelBase::AllSectionSourceObjectName = "AllSection_Source";
 FText UDataHierarchyViewModelBase::AllSectionDisplayName = LOCTEXT("AllSectionDisplayName", "All");
 
-FText FHierarchyElementViewModel::FDefaultUserFeedback::UneditableTarget = LOCTEXT("DefaultUserFeedback_UneditableTarget", "{0} is not editable.");
-FText FHierarchyElementViewModel::FDefaultUserFeedback::TargetCantContainType = LOCTEXT("DefaultUserFeedback_ViewModelCantContainChildType", "{0} can't contain elements of type {1}");
-FText FHierarchyElementViewModel::FDefaultUserFeedback::AddElementFromSourceToHierarchy = LOCTEXT("DefaultUserFeedback_AddElementFromSourceToHierarchy", "Add {0} to {1}");
-FText FHierarchyElementViewModel::FDefaultUserFeedback::MoveElementWithinHierarchy = LOCTEXT("DefaultUserFeedback_MoveElementWithinHierarchy", "Move {0} to {1}");
-
 UHierarchyElement::UHierarchyElement()
 {
 	const uint64 DeterministicSeed = 0x377206996352438CLL;
@@ -43,7 +37,7 @@ UHierarchyElement::UHierarchyElement()
 	Identity.Guids.Add(FGuid::NewDeterministicGuid(PathName.ToView(), DeterministicSeed));
 }
 
-TConstArrayView<const TObjectPtr<UHierarchyElement>> UHierarchyElement::GetChildren() const
+const TArray<UHierarchyElement*>& UHierarchyElement::GetChildren() const
 {
 	return Children;
 }
@@ -169,13 +163,6 @@ bool UHierarchyElement::Modify(bool bAlwaysMarkDirty)
 	return bSavedToTransactionBuffer;
 }
 
-void UHierarchyElement::Serialize(FArchive& Ar)
-{
-	Super::Serialize(Ar);
-	
-	Ar.UsingCustomVersion(FDataHierarchyElementCustomVersion::Guid);
-}
-
 void UHierarchyElement::PostLoad()
 {
 	if(Guid_DEPRECATED.IsValid())
@@ -202,9 +189,35 @@ void UHierarchyElement::PostLoad()
 	Super::PostLoad();
 }
 
+UHierarchySection* UHierarchyRoot::AddSection(FText InNewSectionName, int32 InsertIndex, TSubclassOf<UHierarchySection> SectionClass)
+{
+	TSet<FName> ExistingSectionNames;
+	
+	for(FName& SectionName : GetSections())
+	{
+		ExistingSectionNames.Add(SectionName);
+	}
+	
+	FName NewName = UE::DataHierarchyEditor::GetUniqueName(FName(InNewSectionName.ToString()), ExistingSectionNames);
+	UHierarchySection* NewSectionItem = NewObject<UHierarchySection>(this, SectionClass);
+	NewSectionItem->SetSectionName(NewName);
+	NewSectionItem->SetFlags(RF_Transactional);
+	
+	if(InsertIndex == INDEX_NONE)
+	{
+		Sections.Add(NewSectionItem);
+	}
+	else
+	{
+		Sections.Insert(NewSectionItem, InsertIndex);
+	}
+	
+	return NewSectionItem;
+}
+
 UHierarchySection* UHierarchyRoot::FindSectionByIdentity(FHierarchyElementIdentity SectionIdentity)
 {
-	for(UHierarchySection* Section : GetSectionDataMutable())
+	for(UHierarchySection* Section : Sections)
 	{
 		if(Section->GetPersistentIdentity() == SectionIdentity)
 		{
@@ -225,10 +238,32 @@ void UHierarchyRoot::DuplicateSectionFromOtherRoot(const UHierarchySection& Sect
 	Sections.Add(Cast<UHierarchySection>(StaticDuplicateObject(&SectionToCopy, this)));
 }
 
-TSet<FName> UHierarchyRoot::GetSectionNames() const
+void UHierarchyRoot::RemoveSection(FText SectionName)
+{
+	if(Sections.ContainsByPredicate([SectionName](UHierarchySection* Section)
+	{
+		return Section->GetSectionNameAsText().EqualTo(SectionName);
+	}))
+	{
+		Sections.RemoveAll([SectionName](UHierarchySection* Section)
+		{
+			return Section->GetSectionNameAsText().EqualTo(SectionName);
+		});
+	}
+}
+
+void UHierarchyRoot::RemoveSectionByIdentity(FHierarchyElementIdentity SectionIdentity)
+{
+	Sections.RemoveAll([SectionIdentity](UHierarchySection* Section)
+	{
+		return Section->GetPersistentIdentity() == SectionIdentity;
+	});
+}
+
+TSet<FName> UHierarchyRoot::GetSections() const
 {
 	TSet<FName> OutSections;
-	for(const UHierarchySection* Section : Sections)
+	for(UHierarchySection* Section : Sections)
 	{
 		OutSections.Add(Section->GetSectionName());
 	}
@@ -280,6 +315,27 @@ void UHierarchyRoot::Serialize(FStructuredArchive::FRecord Record)
 	Super::Serialize(Record);
 }
 
+bool FHierarchyCategoryViewModel::IsTopCategoryActive() const
+{
+	if(UHierarchyCategory* Category = GetDataMutable<UHierarchyCategory>())
+	{
+		const UHierarchyCategory* Result = Category;
+		const UHierarchyCategory* TopLevelCategory = Result;
+		
+		for (; TopLevelCategory != nullptr; TopLevelCategory = TopLevelCategory->GetTypedOuter<UHierarchyCategory>() )
+		{
+			if(TopLevelCategory != nullptr)
+			{
+				Result = TopLevelCategory;
+			}
+		}
+		
+		return HierarchyViewModel->IsHierarchySectionActive(Result->GetSection());
+	}
+
+	return false;	
+}
+
 FHierarchyElementViewModel::FResultWithUserFeedback FHierarchyCategoryViewModel::CanContainInternal(TSubclassOf<UHierarchyElement> InHierarchyElement)
 {
 	return (InHierarchyElement->IsChildOf<UHierarchyItem>() || InHierarchyElement->IsChildOf<UHierarchyCategory>());	
@@ -295,75 +351,81 @@ FHierarchyElementViewModel::FResultWithUserFeedback FHierarchyCategoryViewModel:
 	TArray<TSharedPtr<FHierarchyCategoryViewModel>> SiblingCategories;
 	Parent.Pin()->GetChildrenViewModelsForType<UHierarchyCategory, FHierarchyCategoryViewModel>(SiblingCategories);
 	
-	// categories can be dropped on categories, but only if the resulting sibling categories or children categories have different names
-	if(DraggedElement->GetData()->IsA<UHierarchyCategory>())
+	// we only allow drops if some general conditions are fulfilled
+	if(DraggedElement->GetData() != GetData() &&
+		(!DraggedElement->HasParent(AsShared(), false) || ItemDropZone != EItemDropZone::OntoItem)  &&
+		!HasParent(DraggedElement, true))
 	{
-		if(ItemDropZone != EItemDropZone::OntoItem)
+		// categories can be dropped on categories, but only if the resulting sibling categories or children categories have different names
+		if(DraggedElement->GetData()->IsA<UHierarchyCategory>())
 		{
-			bool bContainsSiblingWithSameName = SiblingCategories.ContainsByPredicate([DraggedElement](TSharedPtr<FHierarchyCategoryViewModel> HierarchyCategoryViewModel)
-				{
-					return DraggedElement->ToString() == HierarchyCategoryViewModel->ToString() && DraggedElement != HierarchyCategoryViewModel;
-				});
-
-			if(bContainsSiblingWithSameName)
+			if(ItemDropZone != EItemDropZone::OntoItem)
 			{
-				Results.bResult = false;
-				Results.UserFeedback = LOCTEXT("CantDropCategorNextToCategorySameSiblingNames", "A category of the same name already exists here, potentially in a different section. Please rename your category first.");
-				return Results;
-			}
+				bool bContainsSiblingWithSameName = SiblingCategories.ContainsByPredicate([DraggedElement](TSharedPtr<FHierarchyCategoryViewModel> HierarchyCategoryViewModel)
+					{
+						return DraggedElement->ToString() == HierarchyCategoryViewModel->ToString() && DraggedElement != HierarchyCategoryViewModel;
+					});
 
-			Results.UserFeedback = LOCTEXT("MoveCategoryText", "Move category here");
+				if(bContainsSiblingWithSameName)
+				{
+					Results.bResult = false;
+					Results.UserFeedback = LOCTEXT("CantDropCategorNextToCategorySameSiblingNames", "A category of the same name already exists here, potentially in a different section. Please rename your category first.");
+					return Results;
+				}
 
-			// if we are making a category a sibling of another at the root level, the section will be set to the currently active section. Let that be known.
-			if(Parent.Pin()->GetData()->IsA<UHierarchyRoot>())
-			{
-				UHierarchyCategory* DraggedCategory = Cast<UHierarchyCategory>(DraggedElement->GetDataMutable());
-				
-				FDataHierarchyElementMetaData_SectionAssociation SectionAssociation = DraggedCategory->FindMetaDataOfTypeOrDefault<FDataHierarchyElementMetaData_SectionAssociation>();					
-				if(SectionAssociation.Section != HierarchyViewModel->GetActiveHierarchySectionData())
+				Results.UserFeedback = LOCTEXT("MoveCategoryText", "Move category here");
+
+				// if we are making a category a sibling of another at the root level, the section will be set to the currently active section. Let that be known.
+				if(Parent.Pin()->GetData()->IsA<UHierarchyRoot>())
 				{
-					FText SectionChangeBaseText = LOCTEXT("CategorySectionWillUpdateDueToDrop", "The section of the category will change to {0} after the drop");
-					FText ActualSectionChangeText = FText::FormatOrdered(SectionChangeBaseText, HierarchyViewModel->GetActiveHierarchySectionViewModel()->GetSectionNameAsText());
-					Results.UserFeedback = FText::FormatOrdered(FText::AsCultureInvariant("{0}\n{1}"), Results.UserFeedback.Get(FText::GetEmpty()), ActualSectionChangeText);
+					UHierarchyCategory* DraggedCategory = Cast<UHierarchyCategory>(DraggedElement->GetDataMutable());
+					const UHierarchySection* CurrentElementSection = DraggedCategory->GetSection() == nullptr ? HierarchyViewModel->GetDefaultHierarchySectionViewModel()->GetData<UHierarchySection>() : DraggedCategory->GetSection();
+					
+					if(CurrentElementSection != HierarchyViewModel->GetActiveHierarchySectionData())
+					{
+						FText SectionChangeBaseText = LOCTEXT("CategorySectionWillUpdateDueToDrop", "The section of the category will change to {0} after the drop");
+						FText ActualSectionChangeText = FText::FormatOrdered(SectionChangeBaseText, HierarchyViewModel->GetActiveHierarchySectionViewModel()->GetSectionNameAsText());
+						Results.UserFeedback = FText::FormatOrdered(FText::AsCultureInvariant("{0}\n{1}"), Results.UserFeedback.Get(FText::GetEmpty()), ActualSectionChangeText);
+					}
 				}
 			}
-		}
-		else
-		{
-			bool bContainsChildrenCategoriesWithSameName = TargetChildrenCategories.ContainsByPredicate([DraggedElement](TSharedPtr<FHierarchyCategoryViewModel> HierarchyCategoryViewModel)
+			else
+			{
+				bool bContainsChildrenCategoriesWithSameName = TargetChildrenCategories.ContainsByPredicate([DraggedElement](TSharedPtr<FHierarchyCategoryViewModel> HierarchyCategoryViewModel)
+					{
+						return DraggedElement->ToString() == HierarchyCategoryViewModel->ToString();
+					});
+
+				if(bContainsChildrenCategoriesWithSameName)
 				{
-					return DraggedElement->ToString() == HierarchyCategoryViewModel->ToString();
-				});
+					Results.bResult = false;
+					Results.UserFeedback = LOCTEXT("CantDropCategoryOnCategorySameChildCategoryName", "A sub-category of the same name already exists! Please rename your category first.");
+					return Results;
+				}
 
-			if(bContainsChildrenCategoriesWithSameName)
-			{
-				Results.bResult = false;
-				Results.UserFeedback = LOCTEXT("CantDropCategoryOnCategorySameChildCategoryName", "A sub-category of the same name already exists! Please rename your category first.");
-				return Results;
+				Results.UserFeedback = LOCTEXT("CreateSubcategory", "Drop category here to create a sub-category");
 			}
 
-			Results.UserFeedback = LOCTEXT("CreateSubcategory", "Drop category here to create a sub-category");
+			Results.bResult = true;
+			return Results;
 		}
-
-		Results.bResult = true;
-		return Results;
-	}
-	else if(DraggedElement->GetData()->IsA<UHierarchyItem>())
-	{
-		// items can generally be dropped onto categories
-		Results.bResult = EItemDropZone::OntoItem == ItemDropZone;
-
-		if(Results.bResult)
+		else if(DraggedElement->GetData()->IsA<UHierarchyItem>())
 		{
-			if(DraggedElement->IsForHierarchy() == false)
-			{
-				FText Message = LOCTEXT("AddItemToCategoryDragMessage", "Add {0} to {1}");
-				Results.UserFeedback = FText::FormatOrdered(Message, FText::FromString(DraggedElement->ToString()), FText::FromString(ToString()));
-			}
-			else
+			// items can generally be dropped onto categories
+			Results.bResult = EItemDropZone::OntoItem == ItemDropZone;
+
+			if(Results.bResult)
 			{
-				FText Message = LOCTEXT("MoveItemToCategoryDragMessage", "Move {0} to {1}");
-				Results.UserFeedback = FText::FormatOrdered(Message, FText::FromString(DraggedElement->ToString()), FText::FromString(ToString()));
+				if(DraggedElement->IsForHierarchy() == false)
+				{
+					FText Message = LOCTEXT("AddItemToCategoryDragMessage", "Add {0} to {1}");
+					Results.UserFeedback = FText::FormatOrdered(Message, FText::FromString(DraggedElement->ToString()), FText::FromString(ToString()));
+				}
+				else
+				{
+					FText Message = LOCTEXT("MoveItemToCategoryDragMessage", "Move {0} to {1}");
+					Results.UserFeedback = FText::FormatOrdered(Message, FText::FromString(DraggedElement->ToString()), FText::FromString(ToString()));
+				}
 			}
 		}
 	}
@@ -371,6 +433,18 @@ FHierarchyElementViewModel::FResultWithUserFeedback FHierarchyCategoryViewModel:
 	return Results;
 }
 
+void UHierarchyCategory::FixupSectionLinkage()
+{
+	UHierarchyRoot* OwningRoot = GetTypedOuter<UHierarchyRoot>();
+
+	if(Section != nullptr && Section->GetTypedOuter<UHierarchyRoot>() != OwningRoot)
+	{
+		UHierarchySection* CorrectSection = OwningRoot->FindSectionByIdentity(Section->GetPersistentIdentity());
+		ensure(CorrectSection != nullptr);
+		Section = CorrectSection;
+	}
+}
+
 FHierarchyElementIdentity UHierarchyCategory::ConstructIdentity()
 {
 	FHierarchyElementIdentity Identity;
@@ -388,18 +462,6 @@ void UHierarchyCategory::PostLoad()
 	{
 		SetIdentity(UHierarchyCategory::ConstructIdentity());
 	}
-	
-	const int32 CustomVer = GetLinkerCustomVersion(FDataHierarchyElementCustomVersion::Guid);
-	
-	if(CustomVer < FDataHierarchyElementCustomVersion::Type::MetaDataIntroduction)
-	{
-		if(Section_DEPRECATED)
-		{
-			FDataHierarchyElementMetaData_SectionAssociation* SectionAssociation = FindOrAddMetaDataOfType<FDataHierarchyElementMetaData_SectionAssociation>();
-			SectionAssociation->Section = Section_DEPRECATED;
-			Section_DEPRECATED = nullptr;
-		}
-	}
 }
 
 void UHierarchySection::SetSectionNameAsText(const FText& Text)
@@ -456,6 +518,7 @@ void UDataHierarchyViewModelBase::Initialize()
 	
 	HierarchyRootViewModel->Initialize();
 	HierarchyRootViewModel->AddChildFilter(FHierarchyElementViewModel::FOnFilterChild::CreateUObject(this, &UDataHierarchyViewModelBase::FilterForHierarchySection));
+	HierarchyRootViewModel->AddChildFilter(FHierarchyElementViewModel::FOnFilterChild::CreateUObject(this, &UDataHierarchyViewModelBase::FilterForUncategorizedRootItemsInAllSection));
 	HierarchyRootViewModel->SyncViewModelsToData();
 
 	AllSection = NewObject<UHierarchySection>(this, AllSectionHierarchyObjectName, RF_Transient);
@@ -605,8 +668,23 @@ bool UDataHierarchyViewModelBase::FilterForHierarchySection(TSharedPtr<const FHi
 		}
 
 		// if not, we check against identical section data
-		FDataHierarchyElementMetaData_SectionAssociation SectionAssociation = ViewModel->GetData()->FindMetaDataOfTypeOrDefault<FDataHierarchyElementMetaData_SectionAssociation>();		
-		return ActiveHierarchySectionViewModel.Pin()->GetSectionData() == SectionAssociation.Section;
+		return ActiveHierarchySectionViewModel.Pin()->GetData() == ViewModel->GetSection();
+	}
+
+	return true;
+}
+
+bool UDataHierarchyViewModelBase::FilterForUncategorizedRootItemsInAllSection(TSharedPtr<const FHierarchyElementViewModel> ViewModel) const
+{
+	if(ActiveHierarchySectionViewModel.IsValid())
+	{
+		// we want to filter out all items that are directly added to the root if we aren't in the 'All' section
+		if(ActiveHierarchySectionViewModel.Pin()->IsAllSection())
+		{
+			return true;
+		}
+		
+		return ViewModel->GetData<UHierarchyCategory>() != nullptr;
 	}
 
 	return true;
@@ -723,14 +801,14 @@ TSharedPtr<FHierarchySectionViewModel> UDataHierarchyViewModelBase::GetActiveHie
 	return ActiveHierarchySectionViewModel.Pin();
 }
 
-const UHierarchySection* UDataHierarchyViewModelBase::GetActiveHierarchySectionData() const
+UHierarchySection* UDataHierarchyViewModelBase::GetActiveHierarchySectionData() const
 {
-	return ActiveHierarchySectionViewModel.IsValid() ? ActiveHierarchySectionViewModel.Pin()->GetSectionData() : nullptr;
+	return ActiveHierarchySectionViewModel.IsValid() ? ActiveHierarchySectionViewModel.Pin()->GetDataMutable<UHierarchySection>() : nullptr;
 }
 
 bool UDataHierarchyViewModelBase::IsHierarchySectionActive(const UHierarchySection* Section) const
 {
-	return ActiveHierarchySectionViewModel.IsValid() ? ActiveHierarchySectionViewModel.Pin()->GetSectionData() == Section : false;
+	return ActiveHierarchySectionViewModel.IsValid() ? ActiveHierarchySectionViewModel.Pin()->GetData() == Section : false;
 }
 
 FString UDataHierarchyViewModelBase::OnElementToStringDebug(TSharedPtr<FHierarchyElementViewModel> ElementViewModel) const
@@ -744,24 +822,17 @@ FHierarchyElementViewModel::~FHierarchyElementViewModel()
 	FilteredChildren.Empty();
 }
 
-TSharedPtr<FHierarchyElementViewModel> FHierarchyElementViewModel::AddChild(TSubclassOf<UHierarchyElement> NewChildClass, int32 InsertIndex, FHierarchyElementIdentity ChildIdentity)
+UHierarchyElement* FHierarchyElementViewModel::AddChild(TSubclassOf<UHierarchyElement> NewChildClass, FHierarchyElementIdentity ChildIdentity)
 {
-	UHierarchyElement* NewChild = GetDataMutable()->AddChild<UHierarchyElement>(NewChildClass, InsertIndex);
+	UHierarchyElement* NewChild = NewObject<UHierarchyElement>(GetDataMutable(), NewChildClass);
 	NewChild->SetFlags(RF_Transactional);
-	NewChild->SetIdentity(ChildIdentity);
 	NewChild->Modify();
+	NewChild->SetIdentity(ChildIdentity);
+	GetDataMutable()->GetChildrenMutable().Add(NewChild);
 	
 	SyncViewModelsToData();
-
-	TSharedPtr<FHierarchyElementViewModel> ElementViewModel = FindViewModelForChild(NewChild);
-	ensureMsgf(ElementViewModel.IsValid(), TEXT("Invalid view model detected after syncing view models to data."));
-	
-	TInstancedStruct<FHierarchyElementChangedPayload_AddedElement> Payload = TInstancedStruct<FHierarchyElementChangedPayload_AddedElement>::Make();
-	Payload.GetMutable<FHierarchyElementChangedPayload_AddedElement>().AddedElementViewModel = ElementViewModel;
-	HierarchyViewModel->OnHierarchyElementChanged().Execute(Payload);
-	
 	HierarchyViewModel->OnHierarchyChanged().Broadcast();
-	return ElementViewModel;
+	return NewChild;
 }
 
 void FHierarchyElementViewModel::Tick(float DeltaTime)
@@ -886,12 +957,12 @@ const TArray<TSharedPtr<FHierarchyElementViewModel>>& FHierarchyElementViewModel
 	return FilteredChildren;
 }
 
-void FHierarchyElementViewModel::SortChildrenData()
+void FHierarchyElementViewModel::SortChildrenData() const
 {
 	GetDataMutable()->GetChildrenMutable().StableSort([](const UHierarchyElement& ItemA, const UHierarchyElement& ItemB)
-	{
-		return ItemA.IsA<UHierarchyCategory>() && ItemB.IsA<UHierarchyItem>();
-	});
+		{
+			return ItemA.IsA<UHierarchyCategory>() && ItemB.IsA<UHierarchyItem>();
+		});
 }
 
 int32 FHierarchyElementViewModel::GetHierarchyDepth() const
@@ -967,7 +1038,7 @@ TSharedRef<FHierarchyElementViewModel> FHierarchyElementViewModel::ReparentToThi
 	return ViewModel.ToSharedRef();
 }
 
-TSharedPtr<FHierarchyElementViewModel> FHierarchyElementViewModel::FindViewModelForChild(const UHierarchyElement* Child, bool bSearchRecursively) const
+TSharedPtr<FHierarchyElementViewModel> FHierarchyElementViewModel::FindViewModelForChild(UHierarchyElement* Child, bool bSearchRecursively) const
 {
 	int32 Index = FindIndexOfChild(Child);
 	if(Index != INDEX_NONE)
@@ -1017,7 +1088,7 @@ TSharedPtr<FHierarchyElementViewModel> FHierarchyElementViewModel::FindViewModel
 	return nullptr;
 }
 
-int32 FHierarchyElementViewModel::FindIndexOfChild(const UHierarchyElement* Child) const
+int32 FHierarchyElementViewModel::FindIndexOfChild(UHierarchyElement* Child) const
 {
 	return Children.FindLastByPredicate([Child](TSharedPtr<FHierarchyElementViewModel> Item)
 	{
@@ -1025,14 +1096,14 @@ int32 FHierarchyElementViewModel::FindIndexOfChild(const UHierarchyElement* Chil
 	});
 }
 
-int32 FHierarchyElementViewModel::FindIndexOfDataChild(TSharedPtr<const FHierarchyElementViewModel> Child) const
+int32 FHierarchyElementViewModel::FindIndexOfDataChild(TSharedPtr<FHierarchyElementViewModel> Child) const
 {
-	return GetData()->GetChildren().IndexOfByKey(Child->GetData());
+	return GetData()->GetChildren().Find(Child->GetDataMutable());
 }
 
-int32 FHierarchyElementViewModel::FindIndexOfDataChild(const UHierarchyElement* Child) const
+int32 FHierarchyElementViewModel::FindIndexOfDataChild(UHierarchyElement* Child) const
 {
-	return GetData()->GetChildren().IndexOfByKey(Child);
+	return GetData()->GetChildren().Find(Child);
 }
 
 void FHierarchyElementViewModel::Delete()
@@ -1046,10 +1117,6 @@ void FHierarchyElementViewModel::DeleteChild(TSharedPtr<FHierarchyElementViewMod
 	GetDataMutable()->Modify();
 	GetDataMutable()->GetChildrenMutable().Remove(Child->GetDataMutable());
 	Children.Remove(Child);
-
-	TInstancedStruct<FHierarchyElementChangedPayload_DeletedElement> Payload = TInstancedStruct<FHierarchyElementChangedPayload_DeletedElement>::Make();
-	Payload.GetMutable<FHierarchyElementChangedPayload_DeletedElement>().DeletedElementViewModel = Child;
-	HierarchyViewModel->OnHierarchyElementChanged().Execute(Payload);
 }
 
 TOptional<EItemDropZone> FHierarchyElementViewModel::OnCanRowAcceptDrop(const FDragDropEvent& DragDropEvent, EItemDropZone ItemDropZone, TSharedPtr<FHierarchyElementViewModel> Item)
@@ -1097,58 +1164,29 @@ FHierarchyElementViewModel::FResultWithUserFeedback FHierarchyElementViewModel::
 	return CanDragInternal();
 }
 
-bool FHierarchyElementViewModel::IsExpandedByDefault() const
-{
-	// Otherwise, let the element decide
-	return IsExpandedByDefaultInternal();
-}
-
-FHierarchyElementViewModel::FResultWithUserFeedback FHierarchyElementViewModel::CanDropOn(TSharedPtr<FHierarchyElementViewModel> DraggedElement, EItemDropZone ItemDropZone)
+FHierarchyElementViewModel::FResultWithUserFeedback FHierarchyElementViewModel::CanDropOn(TSharedPtr<FHierarchyElementViewModel> DraggedItem, EItemDropZone ItemDropZone)
 {
 	// Editability check
 	{
 		FResultWithUserFeedback Results = IsEditableByUser();
 		if(Results.bResult == false)
 		{
-			if(Results.UserFeedback.IsSet() == false)
-			{
-				Results.UserFeedback = FText::FormatOrdered(FDefaultUserFeedback::UneditableTarget, this->ToStringAsText());
-			}
-			
 			return Results;
 		}
 	}
 
 	// Self identity check
-	if(DraggedElement->GetData() == GetData())
+	if(DraggedItem->GetData() == GetData())
 	{
 		return false;
 	}
 
-	// Immediate parent check; refuse if we are already parented to the target
-	if(DraggedElement->HasParent(AsShared(), false) && ItemDropZone == EItemDropZone::OntoItem)
-	{
-		return false;
-	}
-
-	// Recursive child check; refuse if we try to drag a parent onto its child
-	if(HasParent(DraggedElement, true))
-	{
-		return false;
-	}
-	
 	// Can Contain checks
 	if(ItemDropZone == EItemDropZone::OntoItem)
 	{
-		FResultWithUserFeedback Results = CanContain(DraggedElement->GetData()->GetClass());
+		FResultWithUserFeedback Results = CanContain(DraggedItem->GetData()->GetClass());
 		if(Results.bResult == false)
 		{
-			if(Results.UserFeedback.IsSet() == false)
-			{
-				Results.UserFeedback = FText::FormatOrdered(FDefaultUserFeedback::TargetCantContainType,
-					GetData()->GetClass()->GetDisplayNameText(), DraggedElement->GetData()->GetClass()->GetDisplayNameText());
-			}
-			
 			return Results;
 		}
 	}
@@ -1156,36 +1194,16 @@ FHierarchyElementViewModel::FResultWithUserFeedback FHierarchyElementViewModel::
 	{
 		if(Parent.IsValid())
 		{
-			FResultWithUserFeedback Results = Parent.Pin()->CanContain(DraggedElement->GetData()->GetClass());
+			FResultWithUserFeedback Results = Parent.Pin()->CanContain(DraggedItem->GetData()->GetClass());
 			if(Results.bResult == false)
 			{
-				if(Results.UserFeedback.IsSet() == false)
-				{
-					Results.UserFeedback = FText::FormatOrdered(FDefaultUserFeedback::TargetCantContainType,
-						GetData()->GetClass()->GetDisplayNameText(), Parent.Pin()->GetData()->GetClass()->GetDisplayNameText());
-				}
-				
 				return Results;
 			}
 		}
 	}
 
 	// Internal check
-	FResultWithUserFeedback Results = CanDropOnInternal(DraggedElement, ItemDropZone);
-	
-	if(Results.bResult && Results.UserFeedback.IsSet() == false)
-	{
-		if(DraggedElement->IsForHierarchy() == false)
-		{
-			Results.UserFeedback = FText::FormatOrdered(FDefaultUserFeedback::AddElementFromSourceToHierarchy, DraggedElement->ToStringAsText(), this->ToStringAsText());
-		}
-		else
-		{
-			Results.UserFeedback = FText::FormatOrdered(FDefaultUserFeedback::MoveElementWithinHierarchy, DraggedElement->ToStringAsText(), this->ToStringAsText());
-		}
-	}
-	
-	return Results;;
+	return CanDropOnInternal(DraggedItem, ItemDropZone);
 }
 
 void FHierarchyElementViewModel::OnDroppedOn(TSharedPtr<FHierarchyElementViewModel> DroppedItem, EItemDropZone ItemDropZone)
@@ -1197,57 +1215,56 @@ void FHierarchyElementViewModel::OnDroppedOn(TSharedPtr<FHierarchyElementViewMod
 	
 	FScopedTransaction Transaction(LOCTEXT("Transaction_MovedItem", "Moved an item in the hierarchy"));
 	HierarchyViewModel->GetHierarchyRoot()->Modify();
-
-	// We calculate the target view model (this if OntoItem, Parent if Above/Below the current item)
-	// and the insert index (INDEX_NONE means adding at the end)
-	TSharedPtr<FHierarchyElementViewModel> TargetViewModel;
-	int32 InsertIndex = INDEX_NONE;
+	
+	bool bDropSucceeded = false;
+	TSharedPtr<FHierarchyElementViewModel> NewItemViewModel;
 	if(ItemDropZone == EItemDropZone::OntoItem)
 	{
-		TargetViewModel = AsShared();
-		InsertIndex = INDEX_NONE;
+		if(DroppedItem->IsForHierarchy() == false)
+		{
+			NewItemViewModel = DuplicateToThis(DroppedItem);
+		}
+		else
+		{
+			NewItemViewModel = ReparentToThis(DroppedItem);
+		}
+
+		bDropSucceeded = true;
 	}
-	else
+	else if(ItemDropZone == EItemDropZone::AboveItem)
 	{
-		if(Parent.IsValid())
-		{
-			TargetViewModel =  Parent.Pin();
-			int32 IndexOfThis = Parent.Pin()->FindIndexOfDataChild(AsShared());
+		int32 IndexOfThis = Parent.Pin()->FindIndexOfDataChild(AsShared());
 
-			// If above an item, we insert at this item's index
-			if(ItemDropZone == EItemDropZone::AboveItem)
-			{
-				InsertIndex = FMath::Max(IndexOfThis, 0);
-			}
-			// If below an item, we insert one index behind this item's index
-			else
-			{
-				InsertIndex = FMath::Min(IndexOfThis+1, Parent.Pin()->GetChildren().Num());
-			}
+		if(DroppedItem->IsForHierarchy() == false)
+		{
+			NewItemViewModel = Parent.Pin()->DuplicateToThis(DroppedItem, FMath::Max(IndexOfThis, 0));
+		}
+		else
+		{
+			NewItemViewModel = Parent.Pin()->ReparentToThis(DroppedItem, FMath::Max(IndexOfThis, 0));
 		}
-	}
 
-	TSharedPtr<FHierarchyElementViewModel> NewItemViewModel;
-	if(ensure(TargetViewModel.IsValid()))
+		bDropSucceeded = true;
+	}
+	else if(ItemDropZone == EItemDropZone::BelowItem)
 	{
-		// If the dropped item is from the source, we duplicate it over
+		int32 IndexOfThis = Parent.Pin()->FindIndexOfDataChild(AsShared());
+		
 		if(DroppedItem->IsForHierarchy() == false)
 		{
-			NewItemViewModel = TargetViewModel->DuplicateToThis(DroppedItem, InsertIndex);
+			NewItemViewModel = Parent.Pin()->DuplicateToThis(DroppedItem, FMath::Min(IndexOfThis+1, Parent.Pin()->GetChildren().Num()));
 		}
-		// If the dropped item is from the hierarchy, we simply move it over
 		else
 		{
-			NewItemViewModel = TargetViewModel->ReparentToThis(DroppedItem, InsertIndex);
+			NewItemViewModel = Parent.Pin()->ReparentToThis(DroppedItem, FMath::Min(IndexOfThis+1, Parent.Pin()->GetChildren().Num()));
 		}
+
+		bDropSucceeded = true;
 	}
 	
-	if(NewItemViewModel.IsValid())
+	if(bDropSucceeded)
 	{
 		NewItemViewModel->PostAddFixup();
-
-		TargetViewModel->PostOnDroppedOn(NewItemViewModel);
-		
 		HierarchyViewModel->RefreshHierarchyView();
 		HierarchyViewModel->RefreshSourceView();
 	}
@@ -1281,13 +1298,7 @@ FHierarchyElementViewModel::FResultWithUserFeedback FHierarchyElementViewModel::
 	return CanContainInternal(HierarchyElementType);
 }
 
-bool FHierarchyElementViewModel::IsExpandedByDefaultInternal() const
-{
-	// By default, we expand all high-level elements
-	return GetHierarchyDepth() <= 1;
-}
-
-FHierarchyElementViewModel::FResultWithUserFeedback FHierarchyElementViewModel::CanDropOnInternal(TSharedPtr<FHierarchyElementViewModel> DraggedElement, EItemDropZone ItemDropZone)
+FHierarchyElementViewModel::FResultWithUserFeedback FHierarchyElementViewModel::CanDropOnInternal(TSharedPtr<FHierarchyElementViewModel>, EItemDropZone ItemDropZone)
 {
 	return true;
 }
@@ -1356,23 +1367,48 @@ void FHierarchyRootViewModel::Initialize()
 	GetOnSynced().BindSP(this, &FHierarchyRootViewModel::PropagateOnSynced);
 }
 
-FHierarchyElementViewModel::FResultWithUserFeedback FHierarchyRootViewModel::CanContainInternal(TSubclassOf<UHierarchyElement> InHierarchyElementType)
+FHierarchyElementViewModel::FResultWithUserFeedback FHierarchyRootViewModel::CanContainInternal(TSubclassOf<UHierarchyElement> InHierarchyElement)
 {
-	// The root generally accepts items and categories, not sections
-	FResultWithUserFeedback Results(InHierarchyElementType->IsChildOf<UHierarchyItem>() || InHierarchyElementType->IsChildOf<UHierarchyCategory>());
-	return Results;
+	return true;
 }
 
 FHierarchyElementViewModel::FResultWithUserFeedback FHierarchyRootViewModel::CanDropOnInternal(TSharedPtr<FHierarchyElementViewModel> DraggedElement, EItemDropZone ItemDropZone)
 {
-	FResultWithUserFeedback Results(true);
-
-	// If we are not in the All section, we only accept categories to drop directly on the root
-	if(HierarchyViewModel->GetActiveHierarchySectionViewModel()->IsAllSection() == false && !DraggedElement->GetData()->IsA<UHierarchyCategory>())
+	FResultWithUserFeedback Results(false);
+	
+	// we only allow drops if some general conditions are fulfilled
+	if(DraggedElement->GetData() != GetData() &&
+		(!DraggedElement->HasParent(AsShared(), false) || ItemDropZone != EItemDropZone::OntoItem)  &&
+		!HasParent(DraggedElement, true))
 	{
-		Results.bResult = false;
-		FText BaseMessage = LOCTEXT("RootCanOnlyContainCategories", "While section {0} is active, the root only accepts categories. Switch to the All section first.");
-		Results.UserFeedback = FText::FormatOrdered(BaseMessage, HierarchyViewModel->GetActiveHierarchySectionViewModel()->GetSectionNameAsText());
+		Results.bResult = 
+			// items can be dropped onto the root directly if the section is set to "All"
+			(DraggedElement->GetData()->IsA<UHierarchyItem>() && HierarchyViewModel->GetActiveHierarchySectionViewModel()->IsAllSection())
+				||
+			// categories can be dropped onto the root always
+			(DraggedElement->GetData()->IsA<UHierarchyCategory>());
+
+		if(Results.bResult)
+		{
+			if(DraggedElement->IsForHierarchy() == false)
+			{
+				FText Message = LOCTEXT("CanDropSourceItemOnRootDragMessage", "Add {0} to the hierarchy root.");
+				Message = FText::FormatOrdered(Message, FText::FromString(DraggedElement->ToString()));
+				Results.UserFeedback = Message;
+			}
+			else
+			{
+				FText Message = LOCTEXT("CanDropHierarchyItemOnRootDragMessage", "Move {0} to the hierarchy root.");
+				Message = FText::FormatOrdered(Message, FText::FromString(DraggedElement->ToString()));
+				Results.UserFeedback = Message;
+			}			
+		}
+		else
+		{
+			FText Message = LOCTEXT("CantDropHierarchyItemOnRootDragMessage", "Can not add {0} here. Please add it to a category!");
+			Message = FText::FormatOrdered(Message, FText::FromString(DraggedElement->ToString()));
+			Results.UserFeedback = Message;
+		}
 	}
 	
 	return Results;
@@ -1380,38 +1416,24 @@ FHierarchyElementViewModel::FResultWithUserFeedback FHierarchyRootViewModel::Can
 
 TSharedPtr<FHierarchySectionViewModel> FHierarchyRootViewModel::AddSection()
 {
-	FScopedTransaction ScopedTransaction(LOCTEXT("NewSectionAdded", "Added Section"));
+	FScopedTransaction ScopedTransaction(LOCTEXT("NewSectionAdded","Added Section"));
 	HierarchyViewModel->GetHierarchyRoot()->Modify();
 	
-	TSet<FName> ExistingSectionNames;
-    	
-	for(FName& SectionName : GetData<UHierarchyRoot>()->GetSectionNames())
+	UHierarchySection* SectionData = GetDataMutable<UHierarchyRoot>()->AddSection(LOCTEXT("HierarchyEditorDefaultNewSectionName", "Section"), 0, HierarchyViewModel->GetSectionDataClass());
+	SectionData->Modify();
+	TSharedPtr<FHierarchyElementViewModel> ViewModel = HierarchyViewModel->CreateViewModelForElement(SectionData, StaticCastSharedRef<FHierarchyRootViewModel>(AsShared()));
+	TSharedPtr<FHierarchySectionViewModel> SectionViewModel = StaticCastSharedPtr<FHierarchySectionViewModel>(ViewModel);
+	if(!ensureMsgf(SectionViewModel.IsValid(), TEXT("Make sure that CreateViewModelForData creates a FHierarchySectionViewModel (or derived) for UHierarchySection elements")))
 	{
-		ExistingSectionNames.Add(SectionName);
+		return nullptr;	
 	}
-
-	FName NewName = UE::DataHierarchyEditor::GetUniqueName(FName(LOCTEXT("HierarchyEditorDefaultNewSectionName", "Section").ToString()), ExistingSectionNames);
-
-	UHierarchySection* NewSection = NewObject<UHierarchySection>(GetDataMutable(), HierarchyViewModel->GetSectionDataClass());
-	GetDataMutable<UHierarchyRoot>()->GetSectionDataMutable().Insert(NewSection, 0);
-	SyncViewModelsToData();
-	
-	TSharedPtr<FHierarchySectionViewModel> SectionViewModel = SectionViewModels[0];
-	ensureMsgf(SectionViewModel->GetData() == NewSection, TEXT("The section view model for a new section should always be at index 0 after syncing view models."));
-	
-	SectionViewModel->SetSectionName(NewName);
-	SectionViewModel->GetDataMutable()->Modify();
 	
+	SectionViewModels.Add(SectionViewModel);
+	SyncViewModelsToData();
 	HierarchyViewModel->SetActiveHierarchySection(SectionViewModel);
-	
+
 	OnSectionAddedDelegate.ExecuteIfBound(SectionViewModel);
-	
-	TInstancedStruct<FHierarchyElementChangedPayload_AddedElement> Payload = TInstancedStruct<FHierarchyElementChangedPayload_AddedElement>::Make();
-	Payload.GetMutable<FHierarchyElementChangedPayload_AddedElement>().AddedElementViewModel = SectionViewModel;
-	HierarchyViewModel->OnHierarchyElementChanged().Execute(Payload);
-	
 	OnSectionsChangedDelegate.ExecuteIfBound();
-	
 	return SectionViewModel;
 }
 
@@ -1422,14 +1444,14 @@ void FHierarchyRootViewModel::DeleteSection(TSharedPtr<FHierarchyElementViewMode
 	SectionViewModels.Remove(SectionViewModel);
 
 	OnSectionDeletedDelegate.ExecuteIfBound(SectionViewModel);
-	
-	TInstancedStruct<FHierarchyElementChangedPayload_DeletedElement> Payload = TInstancedStruct<FHierarchyElementChangedPayload_DeletedElement>::Make();
-	Payload.GetMutable<FHierarchyElementChangedPayload_DeletedElement>().DeletedElementViewModel = SectionViewModel;
-	HierarchyViewModel->OnHierarchyElementChanged().Execute(Payload);
-	
 	OnSectionsChangedDelegate.ExecuteIfBound();
 }
 
+void FHierarchyRootViewModel::PropagateOnSynced()
+{
+	OnSyncPropagatedDelegate.ExecuteIfBound();
+}
+
 void FHierarchyRootViewModel::SyncViewModelsToDataInternal()
 {
 	const UHierarchyRoot* RootData = GetData<UHierarchyRoot>();
@@ -1466,8 +1488,8 @@ void FHierarchyRootViewModel::SyncViewModelsToDataInternal()
 
 		if(SectionViewModel == nullptr)
 		{
-			SectionViewModel = StaticCastSharedPtr<FHierarchySectionViewModel>(HierarchyViewModel->CreateViewModelForElement(Section, StaticCastSharedRef<FHierarchyRootViewModel>(AsShared())));
-			SectionViewModel->SyncViewModelsToData();
+			SectionViewModel = MakeShared<FHierarchySectionViewModel>(Section, StaticCastSharedRef<FHierarchyRootViewModel>(AsShared()), HierarchyViewModel);
+			SectionViewModel->SyncViewModelsToData();;
 		}
 		
 		NewSectionViewModels.Add(SectionViewModel);
@@ -1490,39 +1512,19 @@ void FHierarchyRootViewModel::SyncViewModelsToDataInternal()
 		});
 }
 
-void FHierarchyRootViewModel::PostOnDroppedOn(TSharedPtr<FHierarchyElementViewModel> DroppedElementViewModel)
-{
-	// If we are in the all section, we clear the section association
-	if(HierarchyViewModel->GetActiveHierarchySectionViewModel()->IsAllSection())
-	{
-		DroppedElementViewModel->GetDataMutable()->DeleteMetaDataOfType<FDataHierarchyElementMetaData_SectionAssociation>();
-	}
-	// If not, we update the section association
-	else
-	{
-		FDataHierarchyElementMetaData_SectionAssociation* SectionAssociation = DroppedElementViewModel->GetDataMutable()->FindOrAddMetaDataOfType<FDataHierarchyElementMetaData_SectionAssociation>();
-		SectionAssociation->Section = HierarchyViewModel->GetActiveHierarchySectionViewModel()->GetSectionData();
-	}
-}
-
-void FHierarchyRootViewModel::PropagateOnSynced()
-{
-	OnSyncPropagatedDelegate.ExecuteIfBound();
-}
-
 FHierarchySectionViewModel::FHierarchySectionViewModel(UHierarchySection* InItem, TSharedRef<FHierarchyRootViewModel> InParent, TWeakObjectPtr<UDataHierarchyViewModelBase> InHierarchyViewModel)
 	: FHierarchyElementViewModel(InItem, InParent, InHierarchyViewModel, InParent->IsForHierarchy())
 {
 }
 
-const UHierarchySection* FHierarchySectionViewModel::GetSectionData() const
+UHierarchySection* FHierarchySectionViewModel::GetSectionData() const
 {
 	if(IsAllSection())
 	{
 		return nullptr;
 	}
 
-	return GetData<UHierarchySection>();
+	return GetDataMutable<UHierarchySection>();
 }
 
 FString FHierarchySectionViewModel::ToString() const
@@ -1586,14 +1588,14 @@ bool FHierarchySectionViewModel::CanDeleteInternal()
 	return IsForHierarchy() && !IsAllSection();
 }
 
-bool FHierarchySectionViewModel::OnDroppedOnOverride(TSharedPtr<FHierarchyElementViewModel> DroppedElementViewModel, EItemDropZone ItemDropZone)
+bool FHierarchySectionViewModel::OnDroppedOnOverride(TSharedPtr<FHierarchyElementViewModel> DroppedItem, EItemDropZone ItemDropZone)
 {
-	if(DroppedElementViewModel->GetData()->IsA<UHierarchySection>())
+	if(DroppedItem->GetData()->IsA<UHierarchySection>())
 	{
 		FScopedTransaction Transaction(LOCTEXT("Transaction_OnSectionMoved", "Moved section"));
 		HierarchyViewModel->GetHierarchyRoot()->Modify();
 	
-		UHierarchySection* DraggedSectionData = DroppedElementViewModel->GetDataMutable<UHierarchySection>();
+		UHierarchySection* DraggedSectionData = DroppedItem->GetDataMutable<UHierarchySection>();
 	
 		int32 IndexOfThis = HierarchyViewModel->GetHierarchyRoot()->GetSectionData().Find(GetDataMutable<UHierarchySection>());
 		int32 DraggedSectionIndex = HierarchyViewModel->GetHierarchyRoot()->GetSectionData().Find(DraggedSectionData);
@@ -1624,6 +1626,7 @@ bool FHierarchySectionViewModel::OnDroppedOnOverride(TSharedPtr<FHierarchyElemen
 			}
 	
 			bDropSucceeded = true;
+	
 		}
 	
 		if(bDropSucceeded)
@@ -1631,37 +1634,30 @@ bool FHierarchySectionViewModel::OnDroppedOnOverride(TSharedPtr<FHierarchyElemen
 			HierarchyViewModel->ForceFullRefresh();
 			HierarchyViewModel->OnHierarchyChanged().Broadcast();
 		}
-		
-		return true;
 	}
-
-	if(UHierarchyElement* DroppedElement = DroppedElementViewModel->GetDataMutable())
+	else if(UHierarchyCategory* HierarchyCategory = DroppedItem->GetDataMutable<UHierarchyCategory>())
 	{
-		FScopedTransaction Transaction(LOCTEXT("Transaction_OnSectionDrop", "Moved element to section"));
+		FScopedTransaction Transaction(LOCTEXT("Transaction_OnSectionDrop", "Moved category to section"));
 		HierarchyViewModel->GetHierarchyRoot()->Modify();
 		
-		// If we are in the all section, we clear the section association
-		if(IsAllSection())
-		{
-			DroppedElement->DeleteMetaDataOfType<FDataHierarchyElementMetaData_SectionAssociation>();
-		}
-		// If not, we update the section association
-		else
+		HierarchyCategory->SetSection(IsAllSection() ? nullptr : GetDataMutable<UHierarchySection>());
+	
+		// we null out any sections for all contained categories
+		TArray<UHierarchyCategory*> AllChildCategories;
+		HierarchyCategory->GetChildrenOfType<UHierarchyCategory>(AllChildCategories, true);
+		for(UHierarchyCategory* ChildCategory : AllChildCategories)
 		{
-			FDataHierarchyElementMetaData_SectionAssociation* SectionAssociation = DroppedElement->FindOrAddMetaDataOfType<FDataHierarchyElementMetaData_SectionAssociation>();
-			SectionAssociation->Section = GetSectionData();
+			ChildCategory->SetSection(nullptr);
 		}
-		
+	
 		// we only need to reparent if the parent isn't already the root. This stops unnecessary reordering
-		if(DroppedElementViewModel->GetParent() != HierarchyViewModel->GetHierarchyRootViewModel())
+		if(DroppedItem->GetParent() != HierarchyViewModel->GetHierarchyRootViewModel())
 		{
-			HierarchyViewModel->GetHierarchyRootViewModel()->ReparentToThis(DroppedElementViewModel);
+			HierarchyViewModel->GetHierarchyRootViewModel()->ReparentToThis(DroppedItem);
 		}
 		
 		HierarchyViewModel->RefreshHierarchyView();
 		HierarchyViewModel->OnHierarchyChanged().Broadcast();
-		
-		return true;
 	}
 
 	return true;
@@ -1680,7 +1676,7 @@ FHierarchyElementViewModel::FResultWithUserFeedback FHierarchySectionViewModel::
 	return Results;
 }
 
-FHierarchyElementViewModel::FResultWithUserFeedback FHierarchySectionViewModel::CanDropOnInternal(TSharedPtr<FHierarchyElementViewModel> DraggedElementViewModel, EItemDropZone ItemDropZone)
+FHierarchyElementViewModel::FResultWithUserFeedback FHierarchySectionViewModel::CanDropOnInternal(TSharedPtr<FHierarchyElementViewModel> DraggedElement, EItemDropZone ItemDropZone)
 {
 	FResultWithUserFeedback Results(false);
 	// we don't allow dropping onto source sections and we don't specify a message as the sections aren't going to light up as valid drop targets
@@ -1689,9 +1685,20 @@ FHierarchyElementViewModel::FResultWithUserFeedback FHierarchySectionViewModel::
 		return false;
 	}
 	
-	if(UHierarchySection* DraggedSection = Cast<UHierarchySection>(DraggedElementViewModel->GetDataMutable()))
+	if(const UHierarchyCategory* DraggedCategory = Cast<UHierarchyCategory>(DraggedElement->GetData()))
+	{
+		if(ItemDropZone == EItemDropZone::OntoItem)
+		{
+			FText Message = LOCTEXT("DropCategoryOnSectionDragMessage", "Add {0} to section {1}");
+			Message = FText::FormatOrdered(Message, FText::FromString(DraggedElement->ToString()), FText::FromString(ToString()));
+			
+			Results.bResult = GetData() != DraggedCategory->GetSection();
+			Results.UserFeedback = Results.bResult ? Message : FText::GetEmpty();
+		}
+	}
+	else if(UHierarchySection* DraggedSection = Cast<UHierarchySection>(DraggedElement->GetDataMutable()))
 	{
-		TSharedPtr<FHierarchySectionViewModel> DraggedSectionViewModel = StaticCastSharedPtr<FHierarchySectionViewModel>(DraggedElementViewModel);
+		TSharedPtr<FHierarchySectionViewModel> DraggedSectionViewModel = StaticCastSharedPtr<FHierarchySectionViewModel>(DraggedElement);
 		const bool bSameSection = GetSectionData() == DraggedSectionViewModel->GetSectionData();
 
 		// If we drag a section onto a section, nothing happens
@@ -1722,29 +1729,14 @@ FHierarchyElementViewModel::FResultWithUserFeedback FHierarchySectionViewModel::
 			if(ItemDropZone != EItemDropZone::OntoItem)
 			{
 				FText Message = LOCTEXT("MoveSectionLeftDragMessage", "Move section {0} here");
-				Message = FText::FormatOrdered(Message, FText::FromString(DraggedElementViewModel->ToString()));
+				Message = FText::FormatOrdered(Message, FText::FromString(DraggedElement->ToString()));
 				Results.UserFeedback = Message;
 			}
 		}
 	}
-	// For any element other than sections, we let the CanContain check handle that will be called before this filter out invalid element types
-	// If it's a valid type, we assume it's generally allowed to be dropped onto a section
-	else
+	else if(UHierarchyItem* Item = Cast<UHierarchyItem>(DraggedElement->GetDataMutable()))
 	{
-		if(ItemDropZone == EItemDropZone::OntoItem)
-		{
-			// We only allow adding to a section if the element isn't already part of that section
-			FDataHierarchyElementMetaData_SectionAssociation SectionAssociation = DraggedElementViewModel->GetData()->FindMetaDataOfTypeOrDefault<FDataHierarchyElementMetaData_SectionAssociation>();
-			
-			Results.bResult = SectionAssociation.Section != GetSectionData();
-			if(Results.bResult == true)
-			{
-				FText Message = LOCTEXT("AssociateElementWithSectionDragMessage", "Associate element {0} with section {1}");
-				Results.UserFeedback = FText::FormatOrdered(Message, FText::FromString(DraggedElementViewModel->ToString()), GetSectionNameAsText()); 
-			}
-		}
-		
-		return Results;
+		return false;
 	}
 
 	return Results;
@@ -1756,18 +1748,42 @@ void FHierarchySectionViewModel::FinalizeInternal()
 	{
 		HierarchyViewModel->SetActiveHierarchySection(HierarchyViewModel->GetDefaultHierarchySectionViewModel());
 	}
+
+	// we make sure to reset all categories' section entry that were referencing this section
+	TArray<UHierarchyCategory*> AllCategories;
+	HierarchyViewModel->GetHierarchyRoot()->GetChildrenOfType<UHierarchyCategory>(AllCategories, true);
+
+	for(UHierarchyCategory* Category : AllCategories)
+	{
+		if(Category->GetSection() == GetData())
+		{
+			Category->SetSection(nullptr);
+		}
+	}
 }
 
 FHierarchyElementViewModel::FResultWithUserFeedback FHierarchyItemViewModel::CanContainInternal(TSubclassOf<UHierarchyElement> InHierarchyElement)
 {
-	// By default, items can't contain other children
 	return false;
 }
 
-FHierarchyElementViewModel::FResultWithUserFeedback FHierarchyItemViewModel::CanDropOnInternal(TSharedPtr<FHierarchyElementViewModel> DraggedElement, EItemDropZone ItemDropZone)
+::FHierarchyElementViewModel::FResultWithUserFeedback FHierarchyItemViewModel::CanDropOnInternal(TSharedPtr<FHierarchyElementViewModel> DraggedElement, EItemDropZone ItemDropZone)
 {
-	// items can be generally be dropped above/below other items
-	return (GetData()->IsA<UHierarchyItem>() && ItemDropZone != EItemDropZone::OntoItem);
+	bool bAllowDrop = false;
+	
+	TSharedPtr<FHierarchyElementViewModel> SourceDropItem = DraggedElement;
+	TSharedPtr<FHierarchyElementViewModel> TargetDropItem = AsShared();
+
+	// we only allow drops if some general conditions are fulfilled
+	if(SourceDropItem->GetData() != TargetDropItem->GetData() &&
+		(!SourceDropItem->HasParent(TargetDropItem, false) || ItemDropZone != EItemDropZone::OntoItem)  &&
+		!TargetDropItem->HasParent(SourceDropItem, true))
+	{
+		// items can be generally be dropped above/below other items
+		bAllowDrop = (SourceDropItem->GetData()->IsA<UHierarchyItem>() && ItemDropZone != EItemDropZone::OntoItem);
+	}
+
+	return bAllowDrop;
 }
 
 const FTableRowStyle* FHierarchyItemViewModel::GetRowStyle() const
@@ -1775,6 +1791,39 @@ const FTableRowStyle* FHierarchyItemViewModel::GetRowStyle() const
 	return &FAppStyle::Get().GetWidgetStyle<FTableRowStyle>("TableView.Row");
 }
 
+void FHierarchyCategoryViewModel::PostAddFixup()
+{
+	// If the new category is under the root, we set the section to the currently active one
+	if(Parent.Pin()->GetData()->IsA<UHierarchyRoot>())
+	{
+		UHierarchyCategory* Category = GetDataMutable<UHierarchyCategory>();
+		Category->SetSection(HierarchyViewModel->GetActiveHierarchySectionViewModel()->IsAllSection() ? nullptr : HierarchyViewModel->GetActiveHierarchySectionData());
+
+		// we null out any sections for all contained categories
+		TArray<UHierarchyCategory*> AllChildCategories;
+		Category->GetChildrenOfType<UHierarchyCategory>(AllChildCategories, true);
+		for(UHierarchyCategory* ChildCategory : AllChildCategories)
+		{
+			ChildCategory->SetSection(nullptr);
+		}
+	}
+	// If it's a subcategory, we empty out the section
+	else if(Parent.Pin()->GetData()->IsA<UHierarchyCategory>())
+	{
+		UHierarchyCategory* Category = GetDataMutable<UHierarchyCategory>();
+
+		Category->SetSection(nullptr);
+
+		// we null out any sections for all contained categories
+		TArray<UHierarchyCategory*> AllChildCategories;
+		Category->GetChildrenOfType<UHierarchyCategory>(AllChildCategories, true);
+		for(UHierarchyCategory* ChildCategory : AllChildCategories)
+		{
+			ChildCategory->SetSection(nullptr);
+		}
+	}
+}
+
 const FTableRowStyle* FHierarchyCategoryViewModel::GetRowStyle() const
 {
 	return &FDataHierarchyEditorStyle::Get().GetWidgetStyle<FTableRowStyle>("HierarchyEditor.Row.Category");
@@ -1803,12 +1852,13 @@ void UDataHierarchyViewModelBase::AddCategory(TSharedPtr<FHierarchyElementViewMo
 
 	UClass* CategoryClass = GetCategoryDataClass();
 	
-	TSharedPtr<FHierarchyCategoryViewModel> CategoryViewModel = StaticCastSharedPtr<FHierarchyCategoryViewModel>(CategoryParent->AddChild(CategoryClass, INDEX_NONE, UHierarchyCategory::ConstructIdentity()));
-	if(ensureMsgf(CategoryViewModel.IsValid(), TEXT("Could not find view model for new category of type '%s'. Please ensure your 'CreateViewModelForData' function creates a view model."), *CategoryClass->GetName()))
+	UHierarchyCategory* Category = Cast<UHierarchyCategory>(CategoryParent->AddChild(CategoryClass, UHierarchyCategory::ConstructIdentity()));
+	
+	TSharedPtr<FHierarchyElementViewModel> ViewModel = CategoryParent->FindViewModelForChild(Category, false);
+	if(ensureMsgf(ViewModel.IsValid(), TEXT("Could not find view model for new category of type '%s'. Please ensure your 'CreateViewModelForData' function creates a view model."), *CategoryClass->GetName()))
 	{		
 		TArray<UHierarchyCategory*> SiblingCategories;
-
-		CategoryViewModel->GetParent().Pin()->GetData()->GetChildrenOfType<UHierarchyCategory>(SiblingCategories);
+		Category->GetTypedOuter<UHierarchyElement>()->GetChildrenOfType<UHierarchyCategory>(SiblingCategories);
 		
 		TSet<FName> CategoryNames;
 		for(const auto& SiblingCategory : SiblingCategories)
@@ -1816,14 +1866,24 @@ void UDataHierarchyViewModelBase::AddCategory(TSharedPtr<FHierarchyElementViewMo
 			CategoryNames.Add(SiblingCategory->GetCategoryName());
 		}
 
-		CategoryViewModel->SetCategoryName(UE::DataHierarchyEditor::GetUniqueName(FName("New Category"), CategoryNames));
+		Category->SetCategoryName(UE::DataHierarchyEditor::GetUniqueName(FName("New Category"), CategoryNames));
 		
-		RefreshHierarchyView();
+		// we only set the section property if the current section isn't set to "All"
+		Category->SetSection(GetActiveHierarchySectionViewModel()->IsAllSection() ? nullptr : GetActiveHierarchySectionData());
 		
-		//OnElementAddedDelegate.ExecuteIfBound(ViewModel);
+		RefreshHierarchyView();
+
+		OnElementAddedDelegate.ExecuteIfBound(ViewModel);
 	}
 }
 
+void UDataHierarchyViewModelBase::AddSection() const
+{
+	TSharedPtr<FHierarchySectionViewModel> SectionViewModel = GetHierarchyRootViewModel()->AddSection();	
+	OnElementAddedDelegate.ExecuteIfBound(SectionViewModel);
+	OnHierarchyChangedDelegate.Broadcast();
+}
+
 void UDataHierarchyViewModelBase::GenerateDynamicContextMenu(UToolMenu* ToolMenu)
 {
 	UHierarchyMenuContext* HierarchyMenuContext = ToolMenu->FindContext<UHierarchyMenuContext>();
@@ -1896,11 +1956,11 @@ void UDataHierarchyViewModelBase::GenerateDynamicContextMenuInternal(UToolMenu*
 	}
 }
 
-void UDataHierarchyViewModelBase::AddElementUnderRoot(TSubclassOf<UHierarchyElement> NewChildClass, FHierarchyElementIdentity ChildIdentity)
+UHierarchyElement* UDataHierarchyViewModelBase::AddElementUnderRoot(TSubclassOf<UHierarchyElement> NewChildClass, FHierarchyElementIdentity ChildIdentity)
 {
 	FScopedTransaction Transaction(LOCTEXT("Transaction_AddItem", "Add hierarchy item"));
 	HierarchyRoot->Modify();
-	GetHierarchyRootViewModel()->AddChild(NewChildClass, INDEX_NONE, ChildIdentity);
+	return GetHierarchyRootViewModel()->AddChild(NewChildClass, ChildIdentity);
 }
 
 void UDataHierarchyViewModelBase::DeleteElementWithIdentity(FHierarchyElementIdentity Identity)
@@ -1923,9 +1983,7 @@ void UDataHierarchyViewModelBase::DeleteElementWithIdentity(FHierarchyElementIde
 		}
 	}
 	
-	TArray<TSharedPtr<FHierarchySectionViewModel>> SectionViewModels; 
-	HierarchyRootViewModel->GetChildrenViewModelsForType<UHierarchySection, FHierarchySectionViewModel>(SectionViewModels);
-	
+	TArray<TSharedPtr<FHierarchySectionViewModel>> SectionViewModels = HierarchyRootViewModel->GetSectionViewModels();
 	for(TSharedPtr<FHierarchySectionViewModel> SectionViewModel : SectionViewModels)
 	{
 		if(SectionViewModel->GetData()->GetPersistentIdentity() == Identity && SectionViewModel->CanDelete())
