[Backout] - CL45271379

diff --git a/Engine/Source/Editor/DataHierarchyEditor/Private/Widgets/SDataHierarchyEditor.cpp b/Engine/Source/Editor/DataHierarchyEditor/Private/Widgets/SDataHierarchyEditor.cpp
--- a/Engine/Source/Editor/DataHierarchyEditor/Private/Widgets/SDataHierarchyEditor.cpp
+++ b/Engine/Source/Editor/DataHierarchyEditor/Private/Widgets/SDataHierarchyEditor.cpp
@@ -124,7 +124,6 @@ void SHierarchySection::Construct(const FArguments& InArgs, TSharedPtr<FHierarch
 	[
 		SNew(SHorizontalBox)
 		+ SHorizontalBox::Slot()
-		.VAlign(VAlign_Fill)
 		[
 			SAssignNew(CheckBox, SCheckBox)
 			.Visibility(EVisibility::HitTestInvisible)
@@ -138,14 +137,7 @@ void SHierarchySection::Construct(const FArguments& InArgs, TSharedPtr<FHierarch
 				.AutoWidth()
 				.Padding(2.f)
 				[
-					SNew(SBox)
-					.WidthOverride(20.f)
-					.HeightOverride(20.f)
-					.Visibility(this, &SHierarchySection::GetImageVisibility)
-					[
-						SNew(SImage)
-						.Image(this, &SHierarchySection::GetImageBrush)
-					]
+					SNew(SImage).Image(InSection->GetSectionImage() != nullptr ? InSection->GetSectionImage() : FStyleDefaults::GetNoBrush())
 				]
 				+ SHorizontalBox::Slot()
 				[
@@ -382,16 +374,6 @@ UHierarchySection* SHierarchySection::TryGetSectionData() const
 	return Cast<UHierarchySection>(SectionViewModelWeak.Pin()->GetDataMutable());
 }
 
-const FSlateBrush* SHierarchySection::GetImageBrush() const
-{
-	return (SectionViewModelWeak.IsValid() && SectionViewModelWeak.Pin()->GetSectionImageBrush() != FAppStyle::GetNoBrush()) ? SectionViewModelWeak.Pin()->GetSectionImageBrush() : FAppStyle::GetNoBrush();
-}
-
-EVisibility SHierarchySection::GetImageVisibility() const
-{
-	return (SectionViewModelWeak.IsValid() && SectionViewModelWeak.Pin()->GetSectionImageBrush() != FAppStyle::GetNoBrush()) ? EVisibility::Visible : EVisibility::Collapsed;  
-}
-
 FText SHierarchySection::GetText() const
 {
 	return SectionViewModelWeak.IsValid() ? SectionViewModelWeak.Pin()->GetSectionNameAsText() : FText::GetEmpty();
@@ -539,7 +521,7 @@ void SDataHierarchyEditor::Construct(const FArguments& InArgs, TObjectPtr<UDataH
 	HierarchyViewModel->OnRefreshHierarchyView().BindSP(this, &SDataHierarchyEditor::RefreshHierarchyView);
 	HierarchyViewModel->OnRefreshSectionsView().BindSP(this, &SDataHierarchyEditor::RefreshSectionsView);
 	HierarchyViewModel->OnHierarchySectionActivated().BindSP(this, &SDataHierarchyEditor::OnHierarchySectionActivated);
-	HierarchyViewModel->OnHierarchyElementChanged().BindSP(this, &SDataHierarchyEditor::OnHierarchyElementChanged);
+	HierarchyViewModel->OnElementAdded().BindSP(this, &SDataHierarchyEditor::OnElementAdded);
 
 	BindToHierarchyRootViewModel();
 	
@@ -855,7 +837,7 @@ SDataHierarchyEditor::~SDataHierarchyEditor()
 		HierarchyViewModel->OnRefreshHierarchyView().Unbind();
 		HierarchyViewModel->OnRefreshSectionsView().Unbind();
 		HierarchyViewModel->OnHierarchySectionActivated().Unbind();
-		HierarchyViewModel->OnHierarchyElementChanged().Unbind();
+		HierarchyViewModel->OnElementAdded().Unbind();
 
 		UnbindFromHierarchyRootViewModel();
 
@@ -943,7 +925,7 @@ FReply SDataHierarchyEditor::OnAddCategoryClicked() const
 
 FReply SDataHierarchyEditor::OnAddSectionClicked() const
 {
-	HierarchyViewModel.Get()->GetHierarchyRootViewModel()->AddSection();
+	HierarchyViewModel.Get()->AddSection();
 	return FReply::Handled();
 }
 
@@ -983,8 +965,6 @@ void SDataHierarchyEditor::RefreshSourceView(bool bFullRefresh) const
 	{
 		SourceTreeView->RequestTreeRefresh();
 	}
-	
-	ExpandEntriesByDefault(SourceTreeView);
 }
 
 void SDataHierarchyEditor::RequestRefreshSourceViewNextFrame(bool bFullRefresh)
@@ -1017,8 +997,6 @@ void SDataHierarchyEditor::RefreshHierarchyView(bool bFullRefresh) const
 	{
 		HierarchyTreeView->RequestTreeRefresh();
 	}
-	
-	ExpandEntriesByDefault(HierarchyTreeView);
 }
 
 void SDataHierarchyEditor::RequestRefreshHierarchyViewNextFrame(bool bFullRefresh)
@@ -1055,7 +1033,7 @@ void SDataHierarchyEditor::RefreshSectionsView()
 			SourceSectionBox->AddSlot()
 			.Padding(2.f)
 			.HAlign(HAlign_Center)
-			.VAlign(VAlign_Fill)
+			.VAlign(VAlign_Center)
 			[
 				SectionWidget.ToSharedRef()		
 			];
@@ -1073,7 +1051,7 @@ void SDataHierarchyEditor::RefreshSectionsView()
 			SourceSectionBox->AddSlot()
 			.Padding(2.f)
 			.HAlign(HAlign_Center)
-			.VAlign(VAlign_Fill)
+			.VAlign(VAlign_Center)
 			[
 				DefaultSourceSection.ToSharedRef()
 			];
@@ -1081,7 +1059,7 @@ void SDataHierarchyEditor::RefreshSectionsView()
 	}
 
 	HierarchySectionBox->ClearChildren();
-	for (const TSharedPtr<FHierarchySectionViewModel>& HierarchySection : HierarchyViewModel->GetHierarchyRootViewModel()->GetSectionViewModels())
+	for (TSharedPtr<FHierarchySectionViewModel>& HierarchySection : HierarchyViewModel->GetHierarchyRootViewModel()->GetSectionViewModels())
 	{
 		TSharedPtr<SHierarchySection> SectionWidget = SNew(SHierarchySection, HierarchySection)
 		.IsSectionActive_Lambda([this, HierarchySection]()
@@ -1096,7 +1074,7 @@ void SDataHierarchyEditor::RefreshSectionsView()
 		HierarchySectionBox->AddSlot()
 		.Padding(2.f)
 		.HAlign(HAlign_Center)
-		.VAlign(VAlign_Fill)
+		.VAlign(VAlign_Center)
 		[
 			SectionWidget.ToSharedRef()		
 		];
@@ -1115,7 +1093,7 @@ void SDataHierarchyEditor::RefreshSectionsView()
 	HierarchySectionBox->AddSlot()
 	.Padding(2.f)
 	.HAlign(HAlign_Center)
-	.VAlign(VAlign_Fill)
+	.VAlign(VAlign_Center)
 	[
 		DefaultHierarchySection.ToSharedRef()
 	];
@@ -1261,7 +1239,7 @@ TSharedRef<ITableRow> SDataHierarchyEditor::GenerateHierarchyItemRow(TSharedPtr<
 	];
 }
 
-bool SDataHierarchyEditor::FilterForSourceSection(TSharedPtr<const FHierarchyElementViewModel> ViewModel) const
+bool SDataHierarchyEditor::FilterForSourceSection(TSharedPtr<const FHierarchyElementViewModel> ItemViewModel) const
 {
 	if(ActiveSourceSection.IsValid())
 	{
@@ -1272,8 +1250,7 @@ bool SDataHierarchyEditor::FilterForSourceSection(TSharedPtr<const FHierarchyEle
 		}
 
 		// if not, we check against identical section data
-		FDataHierarchyElementMetaData_SectionAssociation SectionAssociation = ViewModel->GetData()->FindMetaDataOfTypeOrDefault<FDataHierarchyElementMetaData_SectionAssociation>();		
-		return ActiveSourceSection.Pin()->GetSectionData() == SectionAssociation.Section;
+		return ActiveSourceSection.Pin()->GetData() == ItemViewModel->GetSection();
 	}
 
 	return true;
@@ -1462,52 +1439,21 @@ bool SDataHierarchyEditor::CanNavigateToMatchingHierarchyElementFromSelectedSour
 	return false;
 }
 
-void SDataHierarchyEditor::OnHierarchyElementChanged(TInstancedStruct<FHierarchyElementChangedPayload> Payload)
+void SDataHierarchyEditor::OnElementAdded(TSharedPtr<FHierarchyElementViewModel> AddedItem)
 {
-	if(Payload.IsValid())
+	// when a new item is created (opposed to dragged & dropped from source view, i.e. only categories so far)
+	// we make sure to request a tree refresh, select the row, and request a pending rename since the widget will created a frame later
+	if(AddedItem->GetData()->IsA<UHierarchyItem>() || AddedItem->GetData()->IsA<UHierarchyCategory>())
 	{
-		if(Payload.GetScriptStruct() == FHierarchyElementChangedPayload_AddedElement::StaticStruct())
-		{
-			TSharedPtr<FHierarchyElementViewModel> AddedElementViewModel = Payload.GetMutablePtr<FHierarchyElementChangedPayload_AddedElement>()->AddedElementViewModel.Pin();
-			
-			// when a new item is created (opposed to dragged & dropped from source view, i.e. only categories so far)
-			// we make sure to request a tree refresh, select the row, and request a pending rename since the widget will created a frame later
-			if(AddedElementViewModel->GetData()->IsA<UHierarchyItem>() || AddedElementViewModel->GetData()->IsA<UHierarchyCategory>())
-			{
-				HierarchyTreeView->RequestTreeRefresh();
-				NavigateToHierarchyElement(AddedElementViewModel);
-			}
-			else if(AddedElementViewModel->GetData()->IsA<UHierarchySection>())
-			{
-				RefreshSectionsView();
-				HierarchyViewModel->SetActiveHierarchySection(StaticCastSharedPtr<FHierarchySectionViewModel>(AddedElementViewModel));
-			}
-
-			AddedElementViewModel->RequestRenamePending();
-		}
-
-		if(Payload.GetScriptStruct() == FHierarchyElementChangedPayload_DeletedElement::StaticStruct())
-		{
-			TSharedPtr<FHierarchyElementViewModel> DeletedElementViewModel = Payload.GetMutablePtr<FHierarchyElementChangedPayload_DeletedElement>()->DeletedElementViewModel.Pin();
-			
-			// when a new item is created (opposed to dragged & dropped from source view, i.e. only categories so far)
-			// we make sure to request a tree refresh, select the row, and request a pending rename since the widget will created a frame later
-			if(DeletedElementViewModel->GetData()->IsA<UHierarchyItem>() || DeletedElementViewModel->GetData()->IsA<UHierarchyCategory>())
-			{
-				HierarchyTreeView->RequestTreeRefresh();
-				NavigateToHierarchyElement(DeletedElementViewModel);
-			}
-			else if(DeletedElementViewModel->GetData()->IsA<UHierarchySection>())
-			{
-				RefreshSectionsView();
-
-				if(HierarchyViewModel->GetActiveHierarchySectionViewModel() == DeletedElementViewModel)
-				{
-					HierarchyViewModel->SetActiveHierarchySection(HierarchyViewModel->GetDefaultHierarchySectionViewModel());
-				}
-			}
-		}
+		HierarchyTreeView->RequestTreeRefresh();
+		NavigateToHierarchyElement(AddedItem);
 	}
+	else if(AddedItem->GetData()->IsA<UHierarchySection>())
+	{
+		RefreshSectionsView();
+	}
+
+	AddedItem->RequestRenamePending();
 }
 
 void SDataHierarchyEditor::OnHierarchySectionActivated(TSharedPtr<FHierarchySectionViewModel> Section)
@@ -1578,15 +1524,7 @@ TSharedRef<SWidget> SDataHierarchyEditor::OnGetCustomAddContent()
 
 					if(ParentViewModel->CanContain(Type).bResult)
 					{
-						TSharedPtr<FHierarchyElementViewModel> NewElementViewModel = ParentViewModel->AddChild(Type);
-
-						if(const UHierarchySection* Section = DataHierarchyEditor->HierarchyViewModel->GetActiveHierarchySectionData())
-						{
-							FDataHierarchyElementMetaData_SectionAssociation* SectionAssociation = NewElementViewModel->GetDataMutable()->FindOrAddMetaDataOfType<FDataHierarchyElementMetaData_SectionAssociation>();
-							SectionAssociation->Section = Section;
-						}
-
-						NewElementViewModel->SyncViewModelsToData();
+						ParentViewModel->AddChild(Type, FHierarchyElementIdentity());
 					}
 				});
 				UIAction.CanExecuteAction = FToolMenuCanExecuteAction::CreateLambda([ContextObject, Type](const FToolMenuContext& Context)
@@ -1928,26 +1866,6 @@ TOptional<SSearchBox::FSearchResultData> SDataHierarchyEditor::GetSearchResultDa
 	return TOptional<SSearchBox::FSearchResultData>();
 }
 
-void SDataHierarchyEditor::ExpandEntriesByDefault(TSharedPtr<STreeView<TSharedPtr<FHierarchyElementViewModel>>> TreeView) const
-{
-	TArray<TSharedPtr<FHierarchyElementViewModel>> EntriesToProcess(TreeView->GetRootItems());
-	while (EntriesToProcess.Num() > 0)
-	{
-		TSharedPtr<FHierarchyElementViewModel> EntryToProcess = EntriesToProcess[0];
-		EntriesToProcess.RemoveAtSwap(0);
-
-		if (EntryToProcess->IsExpandedByDefault())
-		{
-			TreeView->SetItemExpansion(EntryToProcess, true);
-			EntriesToProcess.Append(EntryToProcess->GetFilteredChildren());
-		}
-		else
-		{
-			TreeView->SetItemExpansion(EntryToProcess, false);
-		}
-	}
-}
-
 FReply SDataHierarchyEditor::HandleHierarchyRootDrop(const FGeometry& Geometry, const FDragDropEvent& DragDropEvent) const
 {
 	if(TSharedPtr<FHierarchyDragDropOp> DragDropOp = DragDropEvent.GetOperationAs<FHierarchyDragDropOp>())
@@ -2010,53 +1928,8 @@ FSlateColor SDataHierarchyEditor::GetRootIconColor() const
 	return FLinearColor(0.2f, 0.2f, 0.2f, 0.5f);
 }
 
-void SDataHierarchyEditor::NotifyPostChange(const FPropertyChangedEvent& PropertyChangedEvent, FProperty* PropertyThatChanged)
+void SDataHierarchyEditor::NotifyPostChange(const FPropertyChangedEvent& PropertyChangedEvent,	FProperty* PropertyThatChanged)
 {
-	const UObject* EditedObject = PropertyChangedEvent.GetObjectBeingEdited(0);
-
-	// We look for the view model corresponding to the edited object.
-	// The object doesn't have to be a hierarchy element, since view models can provide external elements for editing. 
-	if(EditedObject->IsA<UHierarchySection>())
-	{
-		if(SourceRootViewModel.IsValid())
-		{
-			TSharedPtr<FHierarchySectionViewModel>* FoundSourceSectionViewModel = SourceRootViewModel->GetSectionViewModels().FindByPredicate([EditedObject](TSharedPtr<FHierarchySectionViewModel> Candidate)
-				{
-					return Candidate->GetSectionData() == EditedObject;
-				});
-			
-			if(FoundSourceSectionViewModel)
-			{
-				(*FoundSourceSectionViewModel)->PostEditChange(PropertyChangedEvent, PropertyThatChanged);
-			}
-		}
-
-		TSharedPtr<FHierarchySectionViewModel>* FoundHierarchySectionViewModel = HierarchyViewModel->GetHierarchyRootViewModel()->GetSectionViewModels().FindByPredicate([EditedObject](TSharedPtr<FHierarchySectionViewModel> Candidate)
-			{
-				return Candidate->GetSectionData() == EditedObject;
-			});
-
-		if(FoundHierarchySectionViewModel)
-		{
-			(*FoundHierarchySectionViewModel)->PostEditChange(PropertyChangedEvent, PropertyThatChanged);
-		}
-	}
-	else if(EditedObject->IsA<UHierarchyElement>())
-	{
-		if(SourceRootViewModel.IsValid())
-		{
-			if(TSharedPtr<FHierarchyElementViewModel> FoundSourceViewModel = SourceRootViewModel->FindViewModelForChild(Cast<UHierarchyElement>(PropertyChangedEvent.GetObjectBeingEdited(0)), true))
-			{
-				FoundSourceViewModel->PostEditChange(PropertyChangedEvent, PropertyThatChanged);
-			}
-		}
-		
-		if(TSharedPtr<FHierarchyElementViewModel> FoundHierarchyViewModel = HierarchyViewModel->GetHierarchyRootViewModel()->FindViewModelForChild(Cast<UHierarchyElement>(PropertyChangedEvent.GetObjectBeingEdited(0)), true))
-		{
-			FoundHierarchyViewModel->PostEditChange(PropertyChangedEvent, PropertyThatChanged);
-		}
-	}
-	
 	HierarchyViewModel->OnHierarchyPropertiesChanged().Broadcast();
 }
 
diff --git a/Engine/Source/Editor/DataHierarchyEditor/Public/DataHierarchyViewModelBase.h b/Engine/Source/Editor/DataHierarchyEditor/Public/DataHierarchyViewModelBase.h
--- a/Engine/Source/Editor/DataHierarchyEditor/Public/DataHierarchyViewModelBase.h
+++ b/Engine/Source/Editor/DataHierarchyEditor/Public/DataHierarchyViewModelBase.h
@@ -14,8 +14,6 @@
 #include "Templates/SubclassOf.h"
 #include "UObject/Package.h"
 #include "UObject/UObjectGlobals.h"
-#include "StructUtils/InstancedStruct.h"
-#include "DataHierarchyCommonTypes.h"
 #include "DataHierarchyViewModelBase.generated.h"
 
 #define UE_API DATAHIERARCHYEDITOR_API
@@ -63,6 +61,80 @@
  */
 
 struct FHierarchyRootViewModel;
+/** This struct is used to identify a given hierarchy element and can be based on guids and/or names.
+ *  This is particularly useful when a hierarchy element represents an object or a property that is not owned by the hierarchy itself.
+ *  
+ *  
+ */
+USTRUCT()
+struct FHierarchyElementIdentity
+{
+	GENERATED_BODY()
+
+	FHierarchyElementIdentity() {}
+	FHierarchyElementIdentity(TArray<FGuid> InGuids, TArray<FName> InNames) : Guids(InGuids), Names(InNames) {}
+	
+	/** An array of guids that have to be satisfied in order to match. */
+	UPROPERTY()
+	TArray<FGuid> Guids;
+
+	/** Optionally, an array of names can be specified in place of guids. If guids & names are present, guids have to be satisfied first, then names. */
+	UPROPERTY()
+	TArray<FName> Names;
+
+	bool IsValid() const
+	{
+		return Guids.Num() > 0 || Names.Num() > 0;
+	}
+	
+	bool operator==(const FHierarchyElementIdentity& OtherIdentity) const
+	{
+		if(Guids.Num() != OtherIdentity.Guids.Num() || Names.Num() != OtherIdentity.Names.Num())
+		{
+			return false;
+		}
+
+		for(int32 GuidIndex = 0; GuidIndex < Guids.Num(); GuidIndex++)
+		{
+			if(Guids[GuidIndex] != OtherIdentity.Guids[GuidIndex])
+			{
+				return false;
+			}
+		}
+
+		for(int32 NameIndex = 0; NameIndex < Names.Num(); NameIndex++)
+		{
+			if(!Names[NameIndex].IsEqual(OtherIdentity.Names[NameIndex]))
+			{
+				return false;
+			}
+		}
+
+		return true;
+	}
+
+	bool operator!=(const FHierarchyElementIdentity& OtherIdentity) const
+	{
+		return !(*this == OtherIdentity);
+	}
+};
+
+FORCEINLINE uint32 GetTypeHash(const FHierarchyElementIdentity& Identity)
+{
+	uint32 Hash = 0;
+	
+	for(const FGuid& Guid : Identity.Guids)
+	{
+		Hash = HashCombine(Hash, GetTypeHash(Guid));
+	}
+	
+	for(const FName& Name : Identity.Names)
+	{
+		Hash = HashCombine(Hash, GetTypeHash(Name));
+	}
+	
+	return Hash;
+}
 
 /** A base class that is used to refresh data that represents external data. Inherit from this class if you need more context data. */
 UCLASS(MinimalAPI, Transient)
@@ -80,15 +152,12 @@ public:
 	UE_API UHierarchyElement();
 	virtual ~UHierarchyElement() override {}
 
-	UE_API TConstArrayView<const TObjectPtr<UHierarchyElement>> GetChildren() const;
 	TArray<TObjectPtr<UHierarchyElement>>& GetChildrenMutable() { return Children; }
+	UE_API const TArray<UHierarchyElement*>& GetChildren() const;
 
 	template<class ChildClass>
 	ChildClass* AddChild();
 
-	template<class ChildClass>
-	ChildClass* AddChild(TSubclassOf<UHierarchyElement> Class, int32 InsertIndex = INDEX_NONE);
-
 	UE_API UHierarchyElement* FindChildWithIdentity(FHierarchyElementIdentity ChildIdentity, bool bSearchRecursively = false);
 
 	UE_API UHierarchyElement* CopyAndAddItemAsChild(const UHierarchyElement& ItemToCopy);
@@ -118,26 +187,7 @@ public:
 	/** Overridden modify method to also mark all children as modified */
 	UE_API virtual bool Modify(bool bAlwaysMarkDirty = true) override;
 
-	template<typename T>
-	bool HasMetaDataOfType();
-
-	template<typename T>
-	const T* FindMetaDataOfType() const;
-	
-	template<typename T>
-	T* FindMetaDataOfType();
-	
-	template<typename T>
-	T FindMetaDataOfTypeOrDefault() const;
-    	
-	template<typename T>
-	T* FindOrAddMetaDataOfType();
-
-	template <typename T>
-	bool DeleteMetaDataOfType();
-	
 protected:
-	UE_API virtual void Serialize(FArchive& Ar) override;
 	UE_API virtual void PostLoad() override;
 
 	UPROPERTY()
@@ -145,10 +195,6 @@ protected:
 
 	UPROPERTY()
 	FHierarchyElementIdentity Identity;
-
-	/** A map of metadata. Can be used for anything; an example of this is section association. */
-	UPROPERTY()
-	TMap<TObjectPtr<UStruct>, TInstancedStruct<FDataHierarchyElementMetaData>> MetaData;
 	
 	/** An optional guid; can be used if hierarchy items represent outside items */
 	UPROPERTY()
@@ -159,31 +205,11 @@ template <class ChildClass>
 ChildClass* UHierarchyElement::AddChild()
 {
 	ChildClass* NewChild = NewObject<ChildClass>(this);
-	NewChild->SetFlags(RF_Transactional);
-
 	GetChildrenMutable().Add(NewChild);
 
 	return NewChild;
 }
 
-template <class ChildClass>
-ChildClass* UHierarchyElement::AddChild(TSubclassOf<UHierarchyElement> Class, int32 InsertIndex)
-{
-	ChildClass* NewChild = NewObject<ChildClass>(this, Class);
-	NewChild->SetFlags(RF_Transactional);
-
-	if(InsertIndex == INDEX_NONE)
-	{
-		GetChildrenMutable().Add(NewChild);
-	}
-	else
-	{
-		GetChildrenMutable().Insert(NewChild, InsertIndex);
-	}
-
-	return NewChild;
-}
-
 template <class ChildClass>
 bool UHierarchyElement::DoesOneChildExist(bool bRecursive) const
 {
@@ -243,80 +269,6 @@ void UHierarchyElement::SortChildren(const PREDICATE_CLASS& Predicate, bool bRec
 	}
 }
 
-template <typename T>
-bool UHierarchyElement::HasMetaDataOfType()
-{
-	static_assert(TIsDerivedFrom<T, FDataHierarchyElementMetaData>::IsDerived, "T must be derived from FDataHierarchyElementMetaData");
-
-	return MetaData.Contains(T::StaticStruct());
-}
-
-template <typename T>
-const T* UHierarchyElement::FindMetaDataOfType() const
-{
-	static_assert(TIsDerivedFrom<T, FDataHierarchyElementMetaData>::IsDerived, "T must be derived from FDataHierarchyElementMetaData");
-	
-	if(MetaData.Contains(T::StaticStruct()))
-	{
-		return MetaData[T::StaticStruct()].template GetPtr<T>();
-	}
-
-	return nullptr;
-}
-
-template <typename T>
-T* UHierarchyElement::FindMetaDataOfType()
-{
-	static_assert(TIsDerivedFrom<T, FDataHierarchyElementMetaData>::IsDerived, "T must be derived from FDataHierarchyElementMetaData");
-	
-	if(MetaData.Contains(T::StaticStruct()))
-	{
-		return MetaData[T::StaticStruct()].template GetMutablePtr<T>();
-	}
-
-	return nullptr;
-}
-
-template <typename T>
-T UHierarchyElement::FindMetaDataOfTypeOrDefault() const
-{
-	static_assert(TIsDerivedFrom<T, FDataHierarchyElementMetaData>::IsDerived, "T must be derived from FDataHierarchyElementMetaData");
-	static_assert(TIsConstructible<T>::Value, "T must be default constructible");
-	
-	if(MetaData.Contains(T::StaticStruct()))
-	{
-		return MetaData[T::StaticStruct()].template Get<T>();
-	}
-
-	return T();
-}
-
-template <typename T>
-T* UHierarchyElement::FindOrAddMetaDataOfType()
-{
-	static_assert(TIsDerivedFrom<T, FDataHierarchyElementMetaData>::IsDerived, "T must be derived from FDataHierarchyElementMetaData");
-	
-	if(MetaData.Contains(T::StaticStruct()))
-	{
-		return MetaData[T::StaticStruct()].template GetMutablePtr<T>();
-	}
-
-	return MetaData.FindOrAdd(T::StaticStruct(), TInstancedStruct<T>::Make()).template GetMutablePtr<T>();
-}
-
-template <typename T>
-bool UHierarchyElement::DeleteMetaDataOfType()
-{
-	static_assert(TIsDerivedFrom<T, FDataHierarchyElementMetaData>::IsDerived, "T must be derived from FDataHierarchyElementMetaData");
-
-	if(MetaData.Contains(T::StaticStruct()))
-	{
-		return MetaData.Remove(T::StaticStruct()) != 0;
-	}
-
-	return false;
-}
-
 /** A minimal implementation of a section. */
 UCLASS(MinimalAPI)
 class UHierarchySection : public UHierarchyElement
@@ -354,22 +306,25 @@ public:
 	UHierarchyRoot() {}
 	virtual ~UHierarchyRoot() override {}
 
-	TConstArrayView<const TObjectPtr<UHierarchySection>> GetSectionData() const { return Sections; }
+	const TArray<UHierarchySection*>& GetSectionData() const { return Sections; }
 	TArray<TObjectPtr<UHierarchySection>>& GetSectionDataMutable() { return Sections; }
 
-	UE_API TSet<FName> GetSectionNames() const;
+	UE_API TSet<FName> GetSections() const;
 	UE_API int32 GetSectionIndex(UHierarchySection* Section) const;
-	
+
+	UE_API UHierarchySection* AddSection(FText InNewSectionName, int32 InsertIndex = INDEX_NONE, TSubclassOf<UHierarchySection> SectionClass = UHierarchySection::StaticClass());
 	UE_API UHierarchySection* FindSectionByIdentity(FHierarchyElementIdentity SectionIdentity);
 	/** This will copy the section element itself */
 	UE_API void DuplicateSectionFromOtherRoot(const UHierarchySection& SectionToCopy);
-	
-	UE_API virtual void EmptyAllData();
+	UE_API void RemoveSection(FText SectionName);
+	UE_API void RemoveSectionByIdentity(FHierarchyElementIdentity SectionIdentity);
 	
 	UE_API virtual bool Modify(bool bAlwaysMarkDirty = true) override;
-protected:
 	UE_API virtual void Serialize(FStructuredArchive::FRecord Record) override;
+
+	UE_API virtual void EmptyAllData();
 protected:
+
 	UPROPERTY()
 	TArray<TObjectPtr<UHierarchySection>> Sections;
 };
@@ -400,13 +355,21 @@ public:
 	FText GetCategoryAsText() const { return FText::FromName(Category); }
 	FText GetTooltip() const { return Tooltip; }
 
+	void SetSection(UHierarchySection* InSection) { Section = InSection; }
+	const UHierarchySection* GetSection() const { return Section; }
+
 	virtual FString ToString() const override { return Category.ToString(); }
 
+	/** Since the category points to a section object, during merge or copy paste etc. it is possible the section pointer will point at a section from another root.
+	 *  We fix this up by looking through our available sections and match up via persistent identity.
+	 *  This function expects the correct section with the same identity to exist already at the root level
+	 */
+	UE_API void FixupSectionLinkage();
+
 	static UE_API FHierarchyElementIdentity ConstructIdentity();
 
 protected:
 	UE_API virtual void PostLoad() override;
-
 private:
 	UPROPERTY()
 	FName Category;
@@ -416,7 +379,7 @@ private:
 	FText Tooltip;
 
 	UPROPERTY()
-	TObjectPtr<UHierarchySection> Section_DEPRECATED = nullptr;
+	TObjectPtr<UHierarchySection> Section = nullptr;
 };
 
 struct FHierarchyElementViewModel;
@@ -467,37 +430,6 @@ public:
 	TArray<TSharedPtr<FHierarchyElementViewModel>> MenuHierarchyElements;
 };
 
-/** Generic payload struct for when an element in the hierarchy has changed. */
-USTRUCT()
-struct FHierarchyElementChangedPayload
-{
-	GENERATED_BODY()
-};
-
-USTRUCT()
-struct FHierarchyElementChangedPayload_AddedElement : public FHierarchyElementChangedPayload
-{
-	GENERATED_BODY()
-	
-	TWeakPtr<FHierarchyElementViewModel> AddedElementViewModel;
-};
-
-USTRUCT()
-struct FHierarchyElementChangedPayload_DeletedElement : public FHierarchyElementChangedPayload
-{
-	GENERATED_BODY()
-
-	TWeakPtr<FHierarchyElementViewModel> DeletedElementViewModel;
-};
-
-USTRUCT()
-struct FHierarchyElementChangedPayload_ElementPropertyChanged : public FHierarchyElementChangedPayload
-{
-	GENERATED_BODY()
-
-	TWeakPtr<FHierarchyElementViewModel> ElementViewModel;
-};
-
 /** The main controller class for the SHierarchyEditor widget. Defines core hierarchy rules.
  *  Inherit from this and override the required virtual functions, instantiate an object, Initialize it and pass it to the SHierarchyEditor widget. */
 UCLASS(MinimalAPI, Abstract, Transient)
@@ -511,7 +443,6 @@ public:
 	DECLARE_DELEGATE_OneParam(FOnRefreshViewRequested, bool bForceFullRefresh)
 	DECLARE_DELEGATE_OneParam(FOnNavigateToElementIdentityInHierarchyRequested, FHierarchyElementIdentity Identity)
 	DECLARE_DELEGATE_OneParam(FOnNavigateToElementInHierarchyRequested, TSharedPtr<FHierarchyElementViewModel> ElementViewModel)
-	DECLARE_DELEGATE_OneParam(FOnHierarchyElementChanged, TInstancedStruct<FHierarchyElementChangedPayload>)
 
 	GENERATED_BODY()
 
@@ -589,8 +520,9 @@ public:
 	/** Set the refresh context to easily allow Hierarchy Elements representing external data to access whether the external data still exists. */
 	void SetRefreshContext(UHierarchyDataRefreshContext* InContext) { RefreshContext = InContext; }
 	
-	UE_API void AddElementUnderRoot(TSubclassOf<UHierarchyElement> NewChildClass, FHierarchyElementIdentity ChildIdentity);
+	UE_API UHierarchyElement* AddElementUnderRoot(TSubclassOf<UHierarchyElement> NewChildClass, FHierarchyElementIdentity ChildIdentity);
 	UE_API void AddCategory(TSharedPtr<FHierarchyElementViewModel> CategoryParent = nullptr) const;
+	UE_API void AddSection() const;
 
 	/** Delete all specified elements. */
 	UE_API void DeleteElements(TArray<TSharedPtr<FHierarchyElementViewModel>> ViewModels) const;
@@ -622,19 +554,18 @@ public:
 	FSimpleDelegate& OnRefreshSectionsView() { return RefreshSectionsViewDelegate; }
 
 	// Delegates for external systems
-	FOnHierarchyElementChanged& OnHierarchyElementChanged() { return OnHierarchyElementChangedDelegate; }
 	FOnHierarchyChanged& OnHierarchyChanged() { return OnHierarchyChangedDelegate; } 
 	FOnHierarchyChanged& OnHierarchyPropertiesChanged() { return OnHierarchyPropertiesChangedDelegate; } 
-	//FOnElementAdded& OnElementAdded() { return OnElementAddedDelegate; }
+	FOnElementAdded& OnElementAdded() { return OnElementAddedDelegate; }
 	FOnRefreshViewRequested& OnRefreshViewRequested() { return RefreshAllViewsRequestedDelegate; }
 	FOnNavigateToElementIdentityInHierarchyRequested& OnNavigateToElementIdentityInHierarchyRequested() { return OnNavigateToElementIdentityInHierarchyRequestedDelegate; }
 	FOnNavigateToElementInHierarchyRequested& OnNavigateToElementInHierarchyRequested() { return OnNavigateToElementInHierarchyRequestedDelegate; }
 	FSimpleDelegate& OnInitialized() { return OnInitializedDelegate; }
-
+	
 	// Sections
 	UE_API void SetActiveHierarchySection(TSharedPtr<struct FHierarchySectionViewModel>);
 	UE_API TSharedPtr<FHierarchySectionViewModel> GetActiveHierarchySectionViewModel() const;
-	UE_API const UHierarchySection* GetActiveHierarchySectionData() const;
+	UE_API UHierarchySection* GetActiveHierarchySectionData() const;
 	UE_API bool IsHierarchySectionActive(const UHierarchySection* Section) const;
 	FOnSectionActivated& OnHierarchySectionActivated() { return OnHierarchySectionActivatedDelegate; }
 
@@ -646,6 +577,7 @@ protected:
 	UE_API virtual bool MatchesContext(const FTransactionContext& InContext, const TArray<TPair<UObject*, FTransactionObjectEvent>>& TransactionObjectContexts) const override;
 private:
 	UE_API bool FilterForHierarchySection(TSharedPtr<const FHierarchyElementViewModel> ViewModel) const;
+	UE_API bool FilterForUncategorizedRootItemsInAllSection(TSharedPtr<const FHierarchyElementViewModel> ViewModel) const;
 
 	static UE_API void GenerateDynamicContextMenu(UToolMenu* ToolMenu);
 
@@ -657,7 +589,6 @@ private:
 
 	UE_API void ToolMenuNavigateTo(const FToolMenuContext& Context) const;
 	UE_API bool ToolMenuCanNavigateTo(const FToolMenuContext& Context) const;
-
 protected:
 	UPROPERTY()
 	TObjectPtr<UHierarchyRoot> HierarchyRoot;
@@ -682,12 +613,10 @@ protected:
 	FOnNavigateToElementIdentityInHierarchyRequested OnNavigateToElementIdentityInHierarchyRequestedDelegate;
 	FOnNavigateToElementInHierarchyRequested OnNavigateToElementInHierarchyRequestedDelegate;
 	
-	//FOnElementAdded OnElementAddedDelegate;
+	FOnElementAdded OnElementAddedDelegate;
 	FOnSectionActivated OnHierarchySectionActivatedDelegate;
 	FOnSectionActivated OnSourceSectionActivatedDelegate;
 	FOnHierarchyChanged OnHierarchyChangedDelegate;
-	FOnHierarchyElementChanged OnHierarchyElementChangedDelegate;
-
 	FOnHierarchyPropertiesChanged OnHierarchyPropertiesChangedDelegate;
 	
 	FSimpleDelegate OnInitializedDelegate;
@@ -733,11 +662,11 @@ struct FHierarchyElementViewModel : TSharedFromThis<FHierarchyElementViewModel>,
 	
 	UE_API virtual ~FHierarchyElementViewModel() override;
 	
-	UHierarchyElement* GetDataMutable() { return Element; }
+	UHierarchyElement* GetDataMutable() const { return Element; }
 	const UHierarchyElement* GetData() const { return Element; }
 	
 	template<class T>
-	T* GetDataMutable() { return Cast<T>(Element); }
+	T* GetDataMutable() const { return Cast<T>(Element); }
 	
 	template<class T>
 	const T* GetData() const { return Cast<T>(Element); }
@@ -753,8 +682,8 @@ struct FHierarchyElementViewModel : TSharedFromThis<FHierarchyElementViewModel>,
 	UE_API void SyncViewModelsToData();
 	
 	/** Every item view model can define its own sort order for its children. By default we put categories above items. */
-	UE_API virtual void SortChildrenData();
-
+	UE_API virtual void SortChildrenData() const;
+	
 	const TArray<TSharedPtr<FHierarchyElementViewModel>>& GetChildren() const { return Children; }
 	TArray<TSharedPtr<FHierarchyElementViewModel>>& GetChildrenMutable() { return Children; }
 	UE_API const TArray<TSharedPtr<FHierarchyElementViewModel>>& GetFilteredChildren() const;
@@ -762,7 +691,7 @@ struct FHierarchyElementViewModel : TSharedFromThis<FHierarchyElementViewModel>,
 	UE_API void AddChildFilter(FOnFilterChild InFilterChild);
 	
 	template<class DataClass, class ViewModelChildClass>
-	void GetChildrenViewModelsForType(TArray<TSharedPtr<ViewModelChildClass>>& OutChildren, bool bRecursive = false) const;
+	void GetChildrenViewModelsForType(TArray<TSharedPtr<ViewModelChildClass>>& OutChildren, bool bRecursive = false);
 
 	/** Returns the hierarchy depth via number of parents above. */
 	UE_API int32 GetHierarchyDepth() const;
@@ -771,13 +700,13 @@ struct FHierarchyElementViewModel : TSharedFromThis<FHierarchyElementViewModel>,
 	UE_API TSharedRef<FHierarchyElementViewModel> DuplicateToThis(TSharedPtr<FHierarchyElementViewModel> ItemToDuplicate, int32 InsertIndex = INDEX_NONE);
 	UE_API TSharedRef<FHierarchyElementViewModel> ReparentToThis(TSharedPtr<FHierarchyElementViewModel> ItemToMove, int32 InsertIndex = INDEX_NONE);
 
-	UE_API TSharedPtr<FHierarchyElementViewModel> FindViewModelForChild(const UHierarchyElement* Child, bool bSearchRecursively = false) const;
+	UE_API TSharedPtr<FHierarchyElementViewModel> FindViewModelForChild(UHierarchyElement* Child, bool bSearchRecursively = false) const;
 	UE_API TSharedPtr<FHierarchyElementViewModel> FindViewModelForChild(FHierarchyElementIdentity ChildIdentity, bool bSearchRecursively = false) const;
-	UE_API int32 FindIndexOfChild(const UHierarchyElement* Child) const;
-	UE_API int32 FindIndexOfDataChild(TSharedPtr<const FHierarchyElementViewModel> Child) const;
-	UE_API int32 FindIndexOfDataChild(const UHierarchyElement* Child) const;
+	UE_API int32 FindIndexOfChild(UHierarchyElement* Child) const;
+	UE_API int32 FindIndexOfDataChild(TSharedPtr<FHierarchyElementViewModel> Child) const;
+	UE_API int32 FindIndexOfDataChild(UHierarchyElement* Child) const;
 
-	UE_API TSharedPtr<FHierarchyElementViewModel> AddChild(TSubclassOf<UHierarchyElement> NewChildClass, int32 InsertIndex = INDEX_NONE, FHierarchyElementIdentity ChildIdentity = FHierarchyElementIdentity());
+	UE_API UHierarchyElement* AddChild(TSubclassOf<UHierarchyElement> NewChildClass, FHierarchyElementIdentity ChildIdentity);
 	
 	/** Deleting will ask the parent to delete its child */
 	UE_API void Delete();
@@ -809,9 +738,6 @@ struct FHierarchyElementViewModel : TSharedFromThis<FHierarchyElementViewModel>,
 		}
 	};
 
-	/** Called when a property on the underlying data changes. */
-	virtual void PostEditChange(const FPropertyChangedEvent& PropertyChangedEvent, FProperty* PropertyThatChanged) {}
-	
 	/** Should return true if properties are supposed to be editable & needs to be true if typical operations should work on it (renaming, dragging, deleting etc.) */
 	virtual FResultWithUserFeedback IsEditableByUser() { return FResultWithUserFeedback(false); }
 	
@@ -846,10 +772,8 @@ struct FHierarchyElementViewModel : TSharedFromThis<FHierarchyElementViewModel>,
 	/** Returns true if deletable */
 	bool CanDelete() { return IsEditableByUser().bResult && CanDeleteInternal(); }
 
-	UE_API bool IsExpandedByDefault() const;
-	
 	/** Returns true if the given item can be dropped on the given target area. */
-	FResultWithUserFeedback CanDropOn(TSharedPtr<FHierarchyElementViewModel> DraggedElement, EItemDropZone ItemDropZone);
+	FResultWithUserFeedback CanDropOn(TSharedPtr<FHierarchyElementViewModel> DraggedItem, EItemDropZone ItemDropZone);
 
 	/** Gets executed when an item was dropped on this. */
 	void OnDroppedOn(TSharedPtr<FHierarchyElementViewModel> DroppedItem, EItemDropZone ItemDropZone);
@@ -859,6 +783,9 @@ struct FHierarchyElementViewModel : TSharedFromThis<FHierarchyElementViewModel>,
 	/** Returns true if this element can add a new element via the add menu. */
 	FResultWithUserFeedback CanContain(TSubclassOf<UHierarchyElement> HierarchyElementType);
 	
+	/** Determines the section this item belongs to. */
+	const UHierarchySection* GetSection() const { return GetSectionInternal(); }
+	
 	/** For data cleanup that represents external data, this needs to return true in order for live cleanup to work. */
 	virtual bool RepresentsExternalData() const { return false; }
 	/** This function determines whether a hierarchy item that represents that external data should be maintained during data refresh
@@ -867,10 +794,10 @@ struct FHierarchyElementViewModel : TSharedFromThis<FHierarchyElementViewModel>,
 	virtual bool DoesExternalDataStillExist(const UHierarchyDataRefreshContext* Context) const { return false; }
 
 	/** The UObject we display in the details panel when this item is selected. By default it's the hierarchy element the view model represents. */
-	UE_API virtual UObject* GetObjectForEditing() { return Element; }
+	virtual UObject* GetObjectForEditing() { return Element; }
 	/** Source items are transient, which is why we don't allow editing by default.
 	 * This is useful to override if source data points at actual data to edit. */
-	UE_API virtual bool AllowEditingInDetailsPanel() const { return bIsForHierarchy; }
+	virtual bool AllowEditingInDetailsPanel() const { return bIsForHierarchy; }
 
 	bool IsForHierarchy() const { return bIsForHierarchy; }
 
@@ -883,7 +810,6 @@ struct FHierarchyElementViewModel : TSharedFromThis<FHierarchyElementViewModel>,
 	FOnSynced& GetOnSynced() { return OnSyncedDelegate; }
 	FOnChildRequestedDeletion& OnChildRequestedDeletion() { return OnChildRequestedDeletionDelegate; }
 
-	/** Returns the parent view model of the hierarchical parent */
 	TWeakPtr<FHierarchyElementViewModel> GetParent() { return Parent; }
 
 	virtual const FTableRowStyle* GetRowStyle() const { return nullptr; }
@@ -891,11 +817,6 @@ struct FHierarchyElementViewModel : TSharedFromThis<FHierarchyElementViewModel>,
 protected:
 	/** Override this to handle post-add logic (for example when dropping an element on another element, or using the add button). */
 	virtual void PostAddFixup() { }
-	
-	/** Called after an element was dropped onto the calling element. Allows inserting specialized fix up. Not called if OnDroppedOnOverride was called.
-	 *  Note: DroppedElementViewModel is the newly created view model _after_ the drop occured. */
-	UE_API virtual void PostOnDroppedOn(TSharedPtr<FHierarchyElementViewModel> DroppedElementViewModel) { }
-
 private:
 	/** Should return true if this element can contain other elements as children. Can optionally specify a message for failure. */
 	virtual FResultWithUserFeedback CanContainInternal(TSubclassOf<UHierarchyElement> InHierarchyElementType) { return false; }
@@ -911,11 +832,14 @@ private:
 	/** Should return true if deletable. By default, we can delete items in the hierarchy, not in the source. */
 	virtual bool CanDeleteInternal() { return IsForHierarchy(); }
 
-	UE_API virtual bool IsExpandedByDefaultInternal() const;
-	
 	/** Should return true if the given drag drop operation is allowed to succeed. */
-	UE_API virtual FResultWithUserFeedback CanDropOnInternal(TSharedPtr<FHierarchyElementViewModel> DraggedElement, EItemDropZone ItemDropZone);
-private:	
+	UE_API virtual FResultWithUserFeedback CanDropOnInternal(TSharedPtr<FHierarchyElementViewModel>, EItemDropZone ItemDropZone);
+
+	/** Can be overridden to support sections in the source list.
+	 * In the hierarchy only categories can be parented directly to the root, but using this it is possible to add items to custom sections in the source panel.
+	 * This will only work for top-level objects, i.e. anything directly under the root. */
+	virtual const UHierarchySection* GetSectionInternal() const { return nullptr; }
+private:
 	/** Optionally implement this to refresh dependent data. */
 	virtual void RefreshChildrenDataInternal() {}
 	/** Optionally implement this to further customize the view model sync process.
@@ -945,18 +869,10 @@ protected:
 	
 	bool bRenamePending = false;
 	bool bIsForHierarchy = false;
-
-	struct FDefaultUserFeedback
-	{
-		static FText UneditableTarget;
-		static FText TargetCantContainType;
-		static FText AddElementFromSourceToHierarchy;
-		static FText MoveElementWithinHierarchy;
-	};
 };
 
 template <class DataClass, class ViewModelClass>
-void FHierarchyElementViewModel::GetChildrenViewModelsForType(TArray<TSharedPtr<ViewModelClass>>& OutChildren, bool bRecursive) const
+void FHierarchyElementViewModel::GetChildrenViewModelsForType(TArray<TSharedPtr<ViewModelClass>>& OutChildren, bool bRecursive)
 {
 	for(auto& Child : Children)
 	{
@@ -987,7 +903,7 @@ struct FHierarchyRootViewModel : FHierarchyElementViewModel
 	UE_API virtual void Initialize() override;
 
 	virtual FResultWithUserFeedback IsEditableByUser() override { return true; }
-	UE_API virtual FResultWithUserFeedback CanContainInternal(TSubclassOf<UHierarchyElement> InHierarchyElementType) override;
+	UE_API virtual FResultWithUserFeedback CanContainInternal(TSubclassOf<UHierarchyElement> InHierarchyElement) override;
 	UE_API virtual FResultWithUserFeedback CanDropOnInternal(TSharedPtr<FHierarchyElementViewModel> DraggedElement, EItemDropZone ItemDropZone) override;
 
 	virtual FResultWithUserFeedback CanHaveChildren() const override { return true; }
@@ -1007,8 +923,6 @@ struct FHierarchyRootViewModel : FHierarchyElementViewModel
 private:
 	UE_API void PropagateOnSynced();
 	UE_API virtual void SyncViewModelsToDataInternal() override;
-	UE_API virtual void PostOnDroppedOn(TSharedPtr<FHierarchyElementViewModel> DroppedElementViewModel) override;
-private:
 	TArray<TSharedPtr<struct FHierarchySectionViewModel>> SectionViewModels;
 
 	FOnSyncPropagated OnSyncPropagatedDelegate;
@@ -1019,11 +933,12 @@ private:
 
 struct FHierarchySectionViewModel : FHierarchyElementViewModel
 {
-	UE_API FHierarchySectionViewModel(UHierarchySection* InItem, TSharedRef<FHierarchyRootViewModel> InParent, TWeakObjectPtr<UDataHierarchyViewModelBase> InHierarchyViewModel);
+	FHierarchySectionViewModel(UHierarchySection* InItem, TSharedRef<FHierarchyRootViewModel> InParent, TWeakObjectPtr<UDataHierarchyViewModelBase> InHierarchyViewModel);
+	
 	virtual ~FHierarchySectionViewModel() override {}
 
 	/** A custom function to retrieve data; if it's the all section, we want to return nullptr as that's what it serializes as. */
-	UE_API const UHierarchySection* GetSectionData() const;
+	UHierarchySection* GetSectionData() const;
 	
 	UE_API virtual FString ToString() const override;
 	
@@ -1035,11 +950,11 @@ struct FHierarchySectionViewModel : FHierarchyElementViewModel
 	UE_API FText GetSectionTooltip() const;
 
 	void SetSectionImage(const FSlateBrush* InSectionImage) { SectionImage = InSectionImage; }
-	UE_API virtual const FSlateBrush* GetSectionImageBrush() const { return SectionImage ? SectionImage : FAppStyle::GetNoBrush(); }
+	const FSlateBrush* GetSectionImage() const { return SectionImage; }
 	
 	bool IsAllSection() const { return GetData()->GetFName() == UDataHierarchyViewModelBase::AllSectionHierarchyObjectName || GetData()->GetFName() == UDataHierarchyViewModelBase::AllSectionSourceObjectName; }
 protected:
-	UE_API virtual bool AllowEditingInDetailsPanel() const override;
+	virtual bool AllowEditingInDetailsPanel() const override;
 
 	/** Only hierarchy sections are editable */
 	UE_API virtual FResultWithUserFeedback IsEditableByUser() override;
@@ -1084,8 +999,7 @@ struct FHierarchyCategoryViewModel : FHierarchyElementViewModel
 {
 	FHierarchyCategoryViewModel(UHierarchyCategory* InCategory, TSharedRef<FHierarchyElementViewModel> InParent, TWeakObjectPtr<UDataHierarchyViewModelBase> InHierarchyViewModel) : FHierarchyElementViewModel(InCategory, InParent, InHierarchyViewModel, InParent->IsForHierarchy()) {}
 	virtual ~FHierarchyCategoryViewModel() override{}
-
-	void SetCategoryName(FName InText) { GetDataMutable<UHierarchyCategory>()->SetCategoryName(InText); }
+	
 	FText GetCategoryName() const { return GetData<UHierarchyCategory>()->GetCategoryAsText(); }
 	
 	virtual FResultWithUserFeedback IsEditableByUser() override { return FResultWithUserFeedback(true); }
@@ -1093,9 +1007,13 @@ struct FHierarchyCategoryViewModel : FHierarchyElementViewModel
 	virtual FResultWithUserFeedback CanDragInternal() override { return true; }
 	virtual bool CanRenameInternal() override { return true; }
 	virtual void RenameInternal(FName NewName) override { GetDataMutable<UHierarchyCategory>()->SetCategoryName(NewName); }
+	virtual const UHierarchySection* GetSectionInternal() const override { return Cast<UHierarchyCategory>(Element)->GetSection(); }
+
+	UE_API bool IsTopCategoryActive() const;
 
 	UE_API virtual FResultWithUserFeedback CanContainInternal(TSubclassOf<UHierarchyElement> InHierarchyElement) override;
 	UE_API virtual FResultWithUserFeedback CanDropOnInternal(TSharedPtr<FHierarchyElementViewModel>, EItemDropZone ItemDropZone) override;
+	UE_API virtual void PostAddFixup() override;
 
 	UE_API virtual const FTableRowStyle* GetRowStyle() const override;
 };
diff --git a/Engine/Source/Editor/UnrealEd/Private/EditorEngine.cpp b/Engine/Source/Editor/UnrealEd/Private/EditorEngine.cpp
--- a/Engine/Source/Editor/UnrealEd/Private/EditorEngine.cpp
+++ b/Engine/Source/Editor/UnrealEd/Private/EditorEngine.cpp
@@ -1522,8 +1522,7 @@ void UEditorEngine::LoadDefaultEditorModules()
 			TEXT("MassEntityEditor"),
 			TEXT("SVGDistanceField"),
 			TEXT("DataHierarchyEditor"),
-			TEXT("ToolWidgets"),
-			TEXT("UserAssetTagsEditor")
+			TEXT("ToolWidgets")
 		};
 
 	FScopedSlowTask ModuleSlowTask((float)UE_ARRAY_COUNT(ModuleNames));
diff --git a/Engine/Source/Editor/UnrealEd/Private/ObjectTools.cpp b/Engine/Source/Editor/UnrealEd/Private/ObjectTools.cpp
--- a/Engine/Source/Editor/UnrealEd/Private/ObjectTools.cpp
+++ b/Engine/Source/Editor/UnrealEd/Private/ObjectTools.cpp
@@ -4501,9 +4501,6 @@ namespace ObjectTools
 
 				// Remove any metadata from old package pointing to moved objects
 				OldPackage->GetMetaData().RemoveMetaDataOutsidePackage(OldPackage);
-
-				// Allow external systems to deal with metadata
-				UPackage::OnMetaDataTransferRequestedDelegate.Broadcast(OldPackage->GetMetaData(), NewPackage->GetMetaData(), OldPackage, NewPackage);
 			}
 
 			// Migrate any thumbnail from the old package to the new one
diff --git a/Engine/Source/Runtime/CoreUObject/Private/UObject/Package.cpp b/Engine/Source/Runtime/CoreUObject/Private/UObject/Package.cpp
--- a/Engine/Source/Runtime/CoreUObject/Private/UObject/Package.cpp
+++ b/Engine/Source/Runtime/CoreUObject/Private/UObject/Package.cpp
@@ -87,10 +87,6 @@ UPackage::FOnPackageDirtyStateChanged UPackage::PackageDirtyStateChangedEvent;
  */
 UPackage::FOnPackageMarkedDirty UPackage::PackageMarkedDirtyEvent;
 
-#if WITH_METADATA
-UPackage::FOnMetaDataTransferRequestedDelegate UPackage::OnMetaDataTransferRequestedDelegate;
-#endif
-
 PRAGMA_DISABLE_DEPRECATION_WARNINGS // Silence deprecation warnings for deprecated CookedHash and LinkerSave members
 FSavePackageResultStruct::FSavePackageResultStruct()
 	: Result(ESavePackageResult::Error), TotalFileSize(0), SerializedPackageFlags(0)
diff --git a/Engine/Source/Runtime/CoreUObject/Public/UObject/Package.h b/Engine/Source/Runtime/CoreUObject/Public/UObject/Package.h
--- a/Engine/Source/Runtime/CoreUObject/Public/UObject/Package.h
+++ b/Engine/Source/Runtime/CoreUObject/Public/UObject/Package.h
@@ -244,13 +244,6 @@ public:
 	/** delegate type for when a package is about to be saved */
 	DECLARE_MULTICAST_DELEGATE_TwoParams(FPreSavePackageWithContext, class UPackage*, FObjectPreSaveContext);
 
-#if WITH_METADATA
-	/** A delegate that is called when old metadata is supposed to move to new metadata, such as when renaming an asset.
-	 *  Since metadata is not handled by the UPROPERTY-system, data transfer has to be done manually.
-	 *  By default, data is thrown out.  */
-	DECLARE_MULTICAST_DELEGATE_FourParams(FOnMetaDataTransferRequestedDelegate, const FMetaData& OldMetaData, FMetaData& NewMetaData, const UPackage* OldPackage, UPackage* NewPackage);
-#endif
-	
 	UE_DEPRECATED(5.0, "Use PreSavePackageWithContextEvent instead.")
 	COREUOBJECT_API static FPreSavePackage PreSavePackageEvent;
 	/** Delegate to notify subscribers when a package is about to be saved. */
@@ -272,9 +265,6 @@ public:
 	*/
 	COREUOBJECT_API static FOnPackageMarkedDirty PackageMarkedDirtyEvent;
 
-#if WITH_METADATA
-	COREUOBJECT_API static FOnMetaDataTransferRequestedDelegate OnMetaDataTransferRequestedDelegate;
-#endif
 private:
 	/** Used by the editor to determine if a package has been changed. */
 	uint8	bDirty:1;
diff --git a/Engine/Source/Runtime/SlateCore/Private/Styling/StarshipCoreStyle.cpp b/Engine/Source/Runtime/SlateCore/Private/Styling/StarshipCoreStyle.cpp
--- a/Engine/Source/Runtime/SlateCore/Private/Styling/StarshipCoreStyle.cpp
+++ b/Engine/Source/Runtime/SlateCore/Private/Styling/StarshipCoreStyle.cpp
@@ -355,7 +355,6 @@ TSharedRef<ISlateStyle> FStarshipCoreStyle::Create()
 		Style->Set("Icons.AutoFilter", new IMAGE_BRUSH_SVG("Starship/Common/FilterAuto", Icon16x16));
 
 		Style->Set("Icons.Tag", new IMAGE_BRUSH_SVG("Starship/Common/TagSimple_16", Icon16x16));
-		Style->Set("Icons.TagCollection", new IMAGE_BRUSH_SVG("Starship/Common/TagCollection_16", Icon16x16));
 		
 		Style->Set("Icons.Lock", new IMAGE_BRUSH_SVG("Starship/Common/lock", Icon16x16));
 		Style->Set("Icons.Unlock", new IMAGE_BRUSH_SVG("Starship/Common/lock-unlocked", Icon16x16));
