Fixed crash due to unbound delegate

#rnx

diff --git a/Engine/Source/Editor/DataHierarchyEditor/Private/DataHierarchyViewModelBase.cpp b/Engine/Source/Editor/DataHierarchyEditor/Private/DataHierarchyViewModelBase.cpp
--- a/Engine/Source/Editor/DataHierarchyEditor/Private/DataHierarchyViewModelBase.cpp
+++ b/Engine/Source/Editor/DataHierarchyEditor/Private/DataHierarchyViewModelBase.cpp
@@ -758,7 +758,7 @@ TSharedPtr<FHierarchyElementViewModel> FHierarchyElementViewModel::AddChild(TSub
 	
 	TInstancedStruct<FHierarchyElementChangedPayload_AddedElement> Payload = TInstancedStruct<FHierarchyElementChangedPayload_AddedElement>::Make();
 	Payload.GetMutable<FHierarchyElementChangedPayload_AddedElement>().AddedElementViewModel = ElementViewModel;
-	HierarchyViewModel->OnHierarchyElementChanged().Execute(Payload);
+	HierarchyViewModel->OnHierarchyElementChanged().ExecuteIfBound(Payload);
 	
 	HierarchyViewModel->OnHierarchyChanged().Broadcast();
 	return ElementViewModel;
@@ -1049,7 +1049,7 @@ void FHierarchyElementViewModel::DeleteChild(TSharedPtr<FHierarchyElementViewMod
 
 	TInstancedStruct<FHierarchyElementChangedPayload_DeletedElement> Payload = TInstancedStruct<FHierarchyElementChangedPayload_DeletedElement>::Make();
 	Payload.GetMutable<FHierarchyElementChangedPayload_DeletedElement>().DeletedElementViewModel = Child;
-	HierarchyViewModel->OnHierarchyElementChanged().Execute(Payload);
+	HierarchyViewModel->OnHierarchyElementChanged().ExecuteIfBound(Payload);
 }
 
 TOptional<EItemDropZone> FHierarchyElementViewModel::OnCanRowAcceptDrop(const FDragDropEvent& DragDropEvent, EItemDropZone ItemDropZone, TSharedPtr<FHierarchyElementViewModel> Item)
@@ -1420,7 +1420,7 @@ TSharedPtr<FHierarchySectionViewModel> FHierarchyRootViewModel::AddSection()
 	
 	TInstancedStruct<FHierarchyElementChangedPayload_AddedElement> Payload = TInstancedStruct<FHierarchyElementChangedPayload_AddedElement>::Make();
 	Payload.GetMutable<FHierarchyElementChangedPayload_AddedElement>().AddedElementViewModel = SectionViewModel;
-	HierarchyViewModel->OnHierarchyElementChanged().Execute(Payload);
+	HierarchyViewModel->OnHierarchyElementChanged().ExecuteIfBound(Payload);
 	
 	OnSectionsChangedDelegate.ExecuteIfBound();
 	
@@ -1437,7 +1437,7 @@ void FHierarchyRootViewModel::DeleteSection(TSharedPtr<FHierarchyElementViewMode
 	
 	TInstancedStruct<FHierarchyElementChangedPayload_DeletedElement> Payload = TInstancedStruct<FHierarchyElementChangedPayload_DeletedElement>::Make();
 	Payload.GetMutable<FHierarchyElementChangedPayload_DeletedElement>().DeletedElementViewModel = SectionViewModel;
-	HierarchyViewModel->OnHierarchyElementChanged().Execute(Payload);
+	HierarchyViewModel->OnHierarchyElementChanged().ExecuteIfBound(Payload);
 	
 	OnSectionsChangedDelegate.ExecuteIfBound();
 }
