Structural change to Change propagation and listeners:
- Tidied up delegates (fused 2 existing ones, deleted an unused one)
- Added "OnHierarchyChanged" function in the view model. "OnHierarchyChangedDelegate" is only broadcast in there
-- the function is the main access point to react to changes if necessary, but the delegate can still be used
-- "OnHierarchyChangedDelegate" is now only for external systems. SDataHierarchyEditor doesn't bind to it anymore
- Added "FindViewModelForObject" function that is slightly different from FindViewModelForChild. It is able to find view models that represent other objects for editing, like if a UNiagaraScriptVariable gets edited in the hierarchy editor, it will still be able to find the view model for the representing hierarchy element

#Synced-CL 49840944

diff --git a/Engine/Plugins/FX/Niagara/Source/NiagaraEditor/Private/Customizations/UserParameters/NiagaraComponentUserParametersBuilder.cpp b/Engine/Plugins/FX/Niagara/Source/NiagaraEditor/Private/Customizations/UserParameters/NiagaraComponentUserParametersBuilder.cpp
--- a/Engine/Plugins/FX/Niagara/Source/NiagaraEditor/Private/Customizations/UserParameters/NiagaraComponentUserParametersBuilder.cpp
+++ b/Engine/Plugins/FX/Niagara/Source/NiagaraEditor/Private/Customizations/UserParameters/NiagaraComponentUserParametersBuilder.cpp
@@ -284,7 +284,6 @@ FNiagaraComponentUserParametersNodeBuilder::~FNiagaraComponentUserParametersNode
 			if(SystemViewModel.IsValid() && SystemViewModel->GetUserParametersHierarchyViewModel())
 			{
 				SystemViewModel->GetUserParametersHierarchyViewModel()->OnHierarchyChanged().RemoveAll(this);
-				SystemViewModel->GetUserParametersHierarchyViewModel()->OnHierarchyPropertiesChanged().RemoveAll(this);
 			}
 		}
 	}
@@ -502,9 +501,7 @@ void FNiagaraComponentUserParametersNodeBuilder::RegisterRebuildOnHierarchyChang
 			if(SystemViewModel.IsValid())
 			{
 				SystemViewModel->GetUserParametersHierarchyViewModel()->OnHierarchyChanged().RemoveAll(this);
-				SystemViewModel->GetUserParametersHierarchyViewModel()->OnHierarchyChanged().Add(UDataHierarchyViewModelBase::FOnHierarchyChanged::FDelegate::CreateSP(this, &FNiagaraComponentUserParametersNodeBuilder::Rebuild));
-				SystemViewModel->GetUserParametersHierarchyViewModel()->OnHierarchyPropertiesChanged().RemoveAll(this);
-				SystemViewModel->GetUserParametersHierarchyViewModel()->OnHierarchyPropertiesChanged().Add(UDataHierarchyViewModelBase::FOnHierarchyPropertiesChanged::FDelegate::CreateSP(this, &FNiagaraComponentUserParametersNodeBuilder::Rebuild));
+				SystemViewModel->GetUserParametersHierarchyViewModel()->OnHierarchyChanged().Add(UDataHierarchyViewModelBase::FOnHierarchyChanged::FDelegate::CreateSP(this, &FNiagaraComponentUserParametersNodeBuilder::RebuildOnHierarchyChanged));
 			}
 		}
 	}
diff --git a/Engine/Plugins/FX/Niagara/Source/NiagaraEditor/Private/Customizations/UserParameters/NiagaraSystemUserParametersBuilder.cpp b/Engine/Plugins/FX/Niagara/Source/NiagaraEditor/Private/Customizations/UserParameters/NiagaraSystemUserParametersBuilder.cpp
--- a/Engine/Plugins/FX/Niagara/Source/NiagaraEditor/Private/Customizations/UserParameters/NiagaraSystemUserParametersBuilder.cpp
+++ b/Engine/Plugins/FX/Niagara/Source/NiagaraEditor/Private/Customizations/UserParameters/NiagaraSystemUserParametersBuilder.cpp
@@ -64,7 +64,6 @@ FNiagaraSystemUserParameterBuilder::~FNiagaraSystemUserParameterBuilder()
 			if (UNiagaraUserParametersHierarchyViewModel* HierarchyViewModel = SystemViewModel.Pin()->GetUserParametersHierarchyViewModel())
 			{
 				HierarchyViewModel->OnHierarchyChanged().RemoveAll(this);
-				HierarchyViewModel->OnHierarchyPropertiesChanged().RemoveAll(this);
 			}
 		}
 	}
@@ -81,9 +80,7 @@ void FNiagaraSystemUserParameterBuilder::GenerateChildContent(IDetailChildrenBui
 			System->GetExposedParameters().OnStructureChanged().Add(FNiagaraParameterStore::FOnStructureChanged::FDelegate::CreateSP(this, &FNiagaraSystemUserParameterBuilder::Rebuild));
 
 			SystemViewModel.Pin()->GetUserParametersHierarchyViewModel()->OnHierarchyChanged().RemoveAll(this);
-			SystemViewModel.Pin()->GetUserParametersHierarchyViewModel()->OnHierarchyChanged().Add(UDataHierarchyViewModelBase::FOnHierarchyChanged::FDelegate::CreateSP(this, &FNiagaraSystemUserParameterBuilder::Rebuild));
-			SystemViewModel.Pin()->GetUserParametersHierarchyViewModel()->OnHierarchyPropertiesChanged().RemoveAll(this);
-			SystemViewModel.Pin()->GetUserParametersHierarchyViewModel()->OnHierarchyPropertiesChanged().Add(UDataHierarchyViewModelBase::FOnHierarchyPropertiesChanged::FDelegate::CreateSP(this, &FNiagaraSystemUserParameterBuilder::Rebuild));
+			SystemViewModel.Pin()->GetUserParametersHierarchyViewModel()->OnHierarchyChanged().Add(UDataHierarchyViewModelBase::FOnHierarchyChanged::FDelegate::CreateSP(this, &FNiagaraSystemUserParameterBuilder::RebuildOnHierarchyChanged));
 			
 			TSharedPtr<FNiagaraUserParameterPanelViewModel> UserParameterPanelViewModel = SystemViewModel.Pin()->GetUserParameterPanelViewModel();
 			if(ensure(UserParameterPanelViewModel.IsValid()))
diff --git a/Engine/Plugins/FX/Niagara/Source/NiagaraEditor/Private/Customizations/UserParameters/NiagaraUserParametersBuilderBase.h b/Engine/Plugins/FX/Niagara/Source/NiagaraEditor/Private/Customizations/UserParameters/NiagaraUserParametersBuilderBase.h
--- a/Engine/Plugins/FX/Niagara/Source/NiagaraEditor/Private/Customizations/UserParameters/NiagaraUserParametersBuilderBase.h
+++ b/Engine/Plugins/FX/Niagara/Source/NiagaraEditor/Private/Customizations/UserParameters/NiagaraUserParametersBuilderBase.h
@@ -27,6 +27,11 @@ public:
 		OnRebuildChildren.ExecuteIfBound();
 	}
 	
+	void RebuildOnHierarchyChanged(const TInstancedStruct<FHierarchyChangedPayload>& Payload) const
+	{
+		Rebuild();
+	}
+	
 	virtual void SetOnRebuildChildren(FSimpleDelegate InOnRegenerateChildren) override
 	{
 		OnRebuildChildren = InOnRegenerateChildren;
diff --git a/Engine/Plugins/FX/Niagara/Source/NiagaraEditor/Private/Toolkits/NiagaraScriptToolkit.cpp b/Engine/Plugins/FX/Niagara/Source/NiagaraEditor/Private/Toolkits/NiagaraScriptToolkit.cpp
--- a/Engine/Plugins/FX/Niagara/Source/NiagaraEditor/Private/Toolkits/NiagaraScriptToolkit.cpp
+++ b/Engine/Plugins/FX/Niagara/Source/NiagaraEditor/Private/Toolkits/NiagaraScriptToolkit.cpp
@@ -213,7 +213,6 @@ void FNiagaraScriptToolkit::Initialize( const EToolkitMode::Type Mode, const TSh
 	ParametersHierarchyViewModel = NewObject<UNiagaraScriptParametersHierarchyViewModel>();
 	ParametersHierarchyViewModel->Initialize(ScriptViewModel.ToSharedRef());
 	ParametersHierarchyViewModel->OnHierarchyChanged().AddSP(this, &FNiagaraScriptToolkit::OnHierarchyChanged);
-	ParametersHierarchyViewModel->OnHierarchyPropertiesChanged().AddSP(this, &FNiagaraScriptToolkit::OnHierarchyPropertiesChanged);
 	
 	ParameterPanelViewModel = MakeShareable(new FNiagaraScriptToolkitParameterPanelViewModel(ScriptViewModel));
 	ParameterDefinitionsPanelViewModel = MakeShareable(new FNiagaraScriptToolkitParameterDefinitionsPanelViewModel(ScriptViewModel));
@@ -497,16 +496,11 @@ void FNiagaraScriptToolkit::OnVMScriptCompiled(UNiagaraScript*, const FGuid&)
 	}
 }
 
-void FNiagaraScriptToolkit::OnHierarchyChanged()
+void FNiagaraScriptToolkit::OnHierarchyChanged(const TInstancedStruct<FHierarchyChangedPayload>& Payload)
 {
 	ScriptViewModel->GetStandaloneScript().GetScriptData()->GetSource()->MarkNotSynchronized("Input Hierarchy Changed");
 }
 
-void FNiagaraScriptToolkit::OnHierarchyPropertiesChanged()
-{
-	ScriptViewModel->GetStandaloneScript().GetScriptData()->GetSource()->MarkNotSynchronized("Input Hierarchy Properties Changed");
-}
-
 TSharedRef<SDockTab> FNiagaraScriptToolkit::SpawnTabScriptDetails(const FSpawnTabArgs& Args)
 {
 	checkf(Args.GetTabId().TabType == ScriptDetailsTabId, TEXT("Wrong tab ID in NiagaraScriptToolkit"));
diff --git a/Engine/Plugins/FX/Niagara/Source/NiagaraEditor/Private/Toolkits/NiagaraScriptToolkit.h b/Engine/Plugins/FX/Niagara/Source/NiagaraEditor/Private/Toolkits/NiagaraScriptToolkit.h
--- a/Engine/Plugins/FX/Niagara/Source/NiagaraEditor/Private/Toolkits/NiagaraScriptToolkit.h
+++ b/Engine/Plugins/FX/Niagara/Source/NiagaraEditor/Private/Toolkits/NiagaraScriptToolkit.h
@@ -96,8 +96,7 @@ private:
 
 	void OnVMScriptCompiled(UNiagaraScript* InScript, const FGuid& ScriptVersion);
 
-	void OnHierarchyChanged();
-	void OnHierarchyPropertiesChanged();
+	void OnHierarchyChanged(const TInstancedStruct<FHierarchyChangedPayload>& Payload);
 	
 	/** Spawns the tab with the update graph inside */
 	TSharedRef<SDockTab> SpawnTabNodeGraph(const FSpawnTabArgs& Args);
diff --git a/Engine/Plugins/FX/Niagara/Source/NiagaraEditor/Private/ViewModels/HierarchyEditor/NiagaraSummaryViewViewModel.cpp b/Engine/Plugins/FX/Niagara/Source/NiagaraEditor/Private/ViewModels/HierarchyEditor/NiagaraSummaryViewViewModel.cpp
--- a/Engine/Plugins/FX/Niagara/Source/NiagaraEditor/Private/ViewModels/HierarchyEditor/NiagaraSummaryViewViewModel.cpp
+++ b/Engine/Plugins/FX/Niagara/Source/NiagaraEditor/Private/ViewModels/HierarchyEditor/NiagaraSummaryViewViewModel.cpp
@@ -652,7 +652,8 @@ void FNiagaraModuleInputViewModel::AddNativeChildrenInputs()
 		
 		GetDataMutable()->SortChildren(SortChildrenInputs, false);
 		SyncViewModelsToData();
-		HierarchyViewModel->OnHierarchyChanged().Broadcast();
+		
+		HierarchyViewModel->OnHierarchyChanged(TInstancedStruct<FHierarchyChangedPayload>::Make("Native children inputs added"));
 	}
 }
 
diff --git a/Engine/Plugins/FX/Niagara/Source/NiagaraEditor/Private/ViewModels/NiagaraEmitterViewModel.cpp b/Engine/Plugins/FX/Niagara/Source/NiagaraEditor/Private/ViewModels/NiagaraEmitterViewModel.cpp
--- a/Engine/Plugins/FX/Niagara/Source/NiagaraEditor/Private/ViewModels/NiagaraEmitterViewModel.cpp
+++ b/Engine/Plugins/FX/Niagara/Source/NiagaraEditor/Private/ViewModels/NiagaraEmitterViewModel.cpp
@@ -87,7 +87,6 @@ bool FNiagaraEmitterViewModel::Initialize(const FVersionedNiagaraEmitter& InEmit
 	SetSimulation(InSimulation);
 	SummaryViewHierarchyViewModel->Initialize(AsShared());
 	SummaryViewHierarchyViewModel->OnHierarchyChanged().AddSP(this, &FNiagaraEmitterViewModel::OnSummaryViewHierarchyChanged);
-	SummaryViewHierarchyViewModel->OnHierarchyPropertiesChanged().AddSP(this, &FNiagaraEmitterViewModel::OnSummaryViewHierarchyChanged);
 	GetEditorData().OnSummaryViewStateChanged().AddSP(this, &FNiagaraEmitterViewModel::UpdateSelectionOnSummaryViewStateChanged);
 	return true;
 }
@@ -96,7 +95,6 @@ void FNiagaraEmitterViewModel::Reset()
 {
 	SetEmitter(FVersionedNiagaraEmitter());
 	SummaryViewHierarchyViewModel->OnHierarchyChanged().RemoveAll(this);
-	SummaryViewHierarchyViewModel->OnHierarchyPropertiesChanged().RemoveAll(this);
 	SetSimulation(nullptr);
 }
 
@@ -623,7 +621,7 @@ void FNiagaraEmitterViewModel::OnEmitterPropertiesChanged()
 	}
 }
 
-void FNiagaraEmitterViewModel::OnSummaryViewHierarchyChanged()
+void FNiagaraEmitterViewModel::OnSummaryViewHierarchyChanged(const TInstancedStruct<FHierarchyChangedPayload>& Payload)
 {
 	GetEditorData().OnPersistentDataChanged().Broadcast();
 }
diff --git a/Engine/Plugins/FX/Niagara/Source/NiagaraEditor/Private/ViewModels/NiagaraScratchPadScriptViewModel.cpp b/Engine/Plugins/FX/Niagara/Source/NiagaraEditor/Private/ViewModels/NiagaraScratchPadScriptViewModel.cpp
--- a/Engine/Plugins/FX/Niagara/Source/NiagaraEditor/Private/ViewModels/NiagaraScratchPadScriptViewModel.cpp
+++ b/Engine/Plugins/FX/Niagara/Source/NiagaraEditor/Private/ViewModels/NiagaraScratchPadScriptViewModel.cpp
@@ -85,7 +85,6 @@ void FNiagaraScratchPadScriptViewModel::Initialize(UNiagaraScript* Script, UNiag
 	ScriptHierarchyViewModel->Initialize(AsShared());
 
 	ScriptHierarchyViewModel->OnHierarchyChanged().AddSP(this, &FNiagaraScratchPadScriptViewModel::OnHierarchyChanged);
-	ScriptHierarchyViewModel->OnHierarchyPropertiesChanged().AddSP(this, &FNiagaraScratchPadScriptViewModel::OnHierarchyPropertiesChanged);
 }
 
 TSharedPtr<IAssetReferenceFilter> FNiagaraScratchPadScriptViewModel::MakeAssetReferenceFilter() const
@@ -315,16 +314,7 @@ FSimpleDelegate& FNiagaraScratchPadScriptViewModel::OnRequestDiscardChanges()
 	return OnRequestDiscardChangesDelegate;
 }
 
-void FNiagaraScratchPadScriptViewModel::OnHierarchyChanged()
-{
-	if (bHasPendingChanges == false)
-	{
-		bHasPendingChanges = true;
-		OnHasUnappliedChangesChangedDelegate.Broadcast();
-	}
-}
-
-void FNiagaraScratchPadScriptViewModel::OnHierarchyPropertiesChanged()
+void FNiagaraScratchPadScriptViewModel::OnHierarchyChanged(const TInstancedStruct<FHierarchyChangedPayload>& Payload)
 {
 	if (bHasPendingChanges == false)
 	{
diff --git a/Engine/Plugins/FX/Niagara/Source/NiagaraEditor/Private/ViewModels/Stack/NiagaraStackSummaryViewInputCollection.cpp b/Engine/Plugins/FX/Niagara/Source/NiagaraEditor/Private/ViewModels/Stack/NiagaraStackSummaryViewInputCollection.cpp
--- a/Engine/Plugins/FX/Niagara/Source/NiagaraEditor/Private/ViewModels/Stack/NiagaraStackSummaryViewInputCollection.cpp
+++ b/Engine/Plugins/FX/Niagara/Source/NiagaraEditor/Private/ViewModels/Stack/NiagaraStackSummaryViewInputCollection.cpp
@@ -41,8 +41,7 @@ void UNiagaraStackSummaryViewCollection::Initialize(FRequiredEntryData InRequire
 
 	Emitter = InEmitter;
 	GetEmitterViewModel()->GetEditorData().OnSummaryViewStateChanged().AddUObject(this, &UNiagaraStackSummaryViewCollection::OnViewStateChanged);
-	GetEmitterViewModel()->GetSummaryHierarchyViewModel()->OnHierarchyChanged().AddUObject(this, &UNiagaraStackSummaryViewCollection::OnViewStateChanged);
-	GetEmitterViewModel()->GetSummaryHierarchyViewModel()->OnHierarchyPropertiesChanged().AddUObject(this, &UNiagaraStackSummaryViewCollection::OnViewStateChanged);
+	GetEmitterViewModel()->GetSummaryHierarchyViewModel()->OnHierarchyChanged().AddUObject(this, &UNiagaraStackSummaryViewCollection::OnHierarchyChanged);
 
 	AddChildFilter(FOnFilterChild::CreateUObject(this, &UNiagaraStackSummaryViewCollection::FilterOnlyModified));
 }
@@ -58,7 +57,6 @@ void UNiagaraStackSummaryViewCollection::FinalizeInternal()
 	if(EmitterViewModelPtr.IsValid() && EmitterViewModelPtr->GetSummaryHierarchyViewModel() != nullptr)
 	{
 		EmitterViewModelPtr->GetSummaryHierarchyViewModel()->OnHierarchyChanged().RemoveAll(this);
-		EmitterViewModelPtr->GetSummaryHierarchyViewModel()->OnHierarchyPropertiesChanged().RemoveAll(this);
 	}
 
 	Super::FinalizeInternal();
diff --git a/Engine/Plugins/FX/Niagara/Source/NiagaraEditor/Private/Widgets/SNiagaraScriptInputPreviewPanel.cpp b/Engine/Plugins/FX/Niagara/Source/NiagaraEditor/Private/Widgets/SNiagaraScriptInputPreviewPanel.cpp
--- a/Engine/Plugins/FX/Niagara/Source/NiagaraEditor/Private/Widgets/SNiagaraScriptInputPreviewPanel.cpp
+++ b/Engine/Plugins/FX/Niagara/Source/NiagaraEditor/Private/Widgets/SNiagaraScriptInputPreviewPanel.cpp
@@ -340,8 +340,7 @@ void SNiagaraScriptInputPreviewPanel::SetupDelegates()
 {
 	Cast<UNiagaraScriptSource>(ScriptToolkit.Pin()->EditedNiagaraScript.GetScriptData()->GetSource())->NodeGraph->OnParametersChanged().AddSP(this, &SNiagaraScriptInputPreviewPanel::OnParametersChanged);
 
-	ScriptToolkit.Pin()->GetHierarchyViewModel()->OnHierarchyChanged().AddSP(this, &SNiagaraScriptInputPreviewPanel::Refresh);
-	ScriptToolkit.Pin()->GetHierarchyViewModel()->OnHierarchyPropertiesChanged().AddSP(this, &SNiagaraScriptInputPreviewPanel::Refresh);
+	ScriptToolkit.Pin()->GetHierarchyViewModel()->OnHierarchyChanged().AddSP(this, &SNiagaraScriptInputPreviewPanel::RefreshOnHierarchyChanged);
 }
 
 void SNiagaraScriptInputPreviewPanel::RemoveDelegates()
@@ -353,7 +352,6 @@ void SNiagaraScriptInputPreviewPanel::RemoveDelegates()
 		if(UNiagaraScriptParametersHierarchyViewModel* HierarchyViewModel = ScriptToolkit.Pin()->GetHierarchyViewModel())
 		{
 			HierarchyViewModel->OnHierarchyChanged().RemoveAll(this);
-			HierarchyViewModel->OnHierarchyPropertiesChanged().RemoveAll(this);
 		}
 	}
 }
diff --git a/Engine/Plugins/FX/Niagara/Source/NiagaraEditor/Public/ViewModels/NiagaraEmitterViewModel.h b/Engine/Plugins/FX/Niagara/Source/NiagaraEditor/Public/ViewModels/NiagaraEmitterViewModel.h
--- a/Engine/Plugins/FX/Niagara/Source/NiagaraEditor/Public/ViewModels/NiagaraEmitterViewModel.h
+++ b/Engine/Plugins/FX/Niagara/Source/NiagaraEditor/Public/ViewModels/NiagaraEmitterViewModel.h
@@ -157,7 +157,7 @@ private:
 
 	void OnEmitterPropertiesChanged();
 
-	void OnSummaryViewHierarchyChanged();
+	void OnSummaryViewHierarchyChanged(const TInstancedStruct<FHierarchyChangedPayload>& Payload);
 
 	void UpdateSelectionOnSummaryViewStateChanged() const;
 
diff --git a/Engine/Plugins/FX/Niagara/Source/NiagaraEditor/Public/ViewModels/NiagaraScratchPadScriptViewModel.h b/Engine/Plugins/FX/Niagara/Source/NiagaraEditor/Public/ViewModels/NiagaraScratchPadScriptViewModel.h
--- a/Engine/Plugins/FX/Niagara/Source/NiagaraEditor/Public/ViewModels/NiagaraScratchPadScriptViewModel.h
+++ b/Engine/Plugins/FX/Niagara/Source/NiagaraEditor/Public/ViewModels/NiagaraScratchPadScriptViewModel.h
@@ -101,8 +101,7 @@ public:
 
 	NIAGARAEDITOR_API FSimpleDelegate& OnRequestDiscardChanges();
 
-	NIAGARAEDITOR_API void OnHierarchyChanged();
-	NIAGARAEDITOR_API void OnHierarchyPropertiesChanged();
+	NIAGARAEDITOR_API void OnHierarchyChanged(const TInstancedStruct<FHierarchyChangedPayload>& Payload);
 
 	NIAGARAEDITOR_API FOnNodeIDFocusRequested& OnNodeIDFocusRequested();
 	NIAGARAEDITOR_API FOnPinIDFocusRequested& OnPinIDFocusRequested();
diff --git a/Engine/Plugins/FX/Niagara/Source/NiagaraEditor/Public/ViewModels/Stack/NiagaraStackSummaryViewInputCollection.h b/Engine/Plugins/FX/Niagara/Source/NiagaraEditor/Public/ViewModels/Stack/NiagaraStackSummaryViewInputCollection.h
--- a/Engine/Plugins/FX/Niagara/Source/NiagaraEditor/Public/ViewModels/Stack/NiagaraStackSummaryViewInputCollection.h
+++ b/Engine/Plugins/FX/Niagara/Source/NiagaraEditor/Public/ViewModels/Stack/NiagaraStackSummaryViewInputCollection.h
@@ -36,6 +36,7 @@ protected:
 	bool FilterOnlyModified(const UNiagaraStackEntry& Child) const;
 private:
 	NIAGARAEDITOR_API void OnViewStateChanged();
+	void OnHierarchyChanged(const TInstancedStruct<FHierarchyChangedPayload>& Payload) { OnViewStateChanged(); }
 private:
 
 	FVersionedNiagaraEmitterWeakPtr Emitter;
diff --git a/Engine/Plugins/FX/Niagara/Source/NiagaraEditor/Public/Widgets/SNiagaraScriptInputPreviewPanel.h b/Engine/Plugins/FX/Niagara/Source/NiagaraEditor/Public/Widgets/SNiagaraScriptInputPreviewPanel.h
--- a/Engine/Plugins/FX/Niagara/Source/NiagaraEditor/Public/Widgets/SNiagaraScriptInputPreviewPanel.h
+++ b/Engine/Plugins/FX/Niagara/Source/NiagaraEditor/Public/Widgets/SNiagaraScriptInputPreviewPanel.h
@@ -30,6 +30,8 @@ public:
 	UE_API virtual ~SNiagaraScriptInputPreviewPanel() override;
 	
 	UE_API void Refresh();
+	void RefreshOnHierarchyChanged(const TInstancedStruct<FHierarchyChangedPayload>& Payload) { Refresh(); }
+	
 	UE_API void SetupDelegates();
 	UE_API void RemoveDelegates();
 private:
diff --git a/Engine/Source/Editor/DataHierarchyEditor/Internal/Widgets/SDataHierarchyEditor.h b/Engine/Source/Editor/DataHierarchyEditor/Internal/Widgets/SDataHierarchyEditor.h
--- a/Engine/Source/Editor/DataHierarchyEditor/Internal/Widgets/SDataHierarchyEditor.h
+++ b/Engine/Source/Editor/DataHierarchyEditor/Internal/Widgets/SDataHierarchyEditor.h
@@ -161,7 +161,6 @@ private:
 	UE_API void NavigateToMatchingHierarchyElementFromSelectedSourceElement() const;
 	UE_API bool CanNavigateToMatchingHierarchyElementFromSelectedSourceElement() const;
 	
-	UE_API void OnHierarchyElementChanged(TInstancedStruct<FHierarchyElementChangedPayload> Payload);
 	UE_API void OnHierarchySectionActivated(TSharedPtr<FHierarchySectionViewModel> Section);
 	UE_API void OnSourceSectionActivated(TSharedPtr<FHierarchySectionViewModel> Section);
 	UE_API void OnHierarchySectionAdded(TSharedPtr<FHierarchySectionViewModel> AddedSection);
diff --git a/Engine/Source/Editor/DataHierarchyEditor/Private/DataHierarchyViewModelBase.cpp b/Engine/Source/Editor/DataHierarchyEditor/Private/DataHierarchyViewModelBase.cpp
--- a/Engine/Source/Editor/DataHierarchyEditor/Private/DataHierarchyViewModelBase.cpp
+++ b/Engine/Source/Editor/DataHierarchyEditor/Private/DataHierarchyViewModelBase.cpp
@@ -515,6 +515,87 @@ TSharedPtr<FHierarchyElementViewModel> UDataHierarchyViewModelBase::CreateViewMo
 	return nullptr;
 }
 
+void UDataHierarchyViewModelBase::OnHierarchyChanged(const TInstancedStruct<FHierarchyChangedPayload>& Payload)
+{
+	if(Payload.IsValid())
+	{
+		if(Payload.GetScriptStruct() == FHierarchyChangedPayload_ElementsAdded::StaticStruct())
+		{
+			TArray<TSharedPtr<FHierarchyElementViewModel>> AddedElementViewModels = Payload.GetPtr<FHierarchyChangedPayload_ElementsAdded>()->AddedElementViewModels;
+			
+			bool bContainsHierarchyElement = AddedElementViewModels.ContainsByPredicate([](TSharedPtr<FHierarchyElementViewModel> Candidate)
+			{
+				return Candidate->GetData()->IsA<UHierarchyItem>() || Candidate->GetData()->IsA<UHierarchyCategory>();
+			});
+			
+			bool bContainsSection = AddedElementViewModels.ContainsByPredicate([](TSharedPtr<FHierarchyElementViewModel> Candidate)
+			{
+				return Candidate->GetData()->IsA<UHierarchySection>();
+			});
+			
+			// when a new item is created (opposed to dragged & dropped from source view, i.e. only categories so far)
+			// we make sure to request a tree refresh, select the row, and request a pending rename since the widget will be created a frame later
+			if(bContainsHierarchyElement)
+			{
+				RefreshHierarchyView(false);
+				
+				if (AddedElementViewModels.Num() == 1)
+				{
+					NavigateToElementInHierarchy(AddedElementViewModels[0].ToSharedRef());
+				}
+			}
+			else if(bContainsSection)
+			{
+				RefreshSectionsView();
+				
+				if (AddedElementViewModels.Num() == 1)
+				{
+					SetActiveHierarchySection(StaticCastSharedPtr<FHierarchySectionViewModel>(AddedElementViewModels[0]));
+				}
+			}
+
+			if (AddedElementViewModels.Num() == 1)
+			{
+				AddedElementViewModels[0]->RequestRenamePending();
+			}
+		}
+
+		if(Payload.GetScriptStruct() == FHierarchyChangedPayload_ElementsDeleted::StaticStruct())
+		{
+			TArray<TSharedPtr<FHierarchyElementViewModel>> DeletedElementViewModels = Payload.GetPtr<FHierarchyChangedPayload_ElementsDeleted>()->DeletedElementViewModels;
+			
+			bool bContainsHierarchyElement = DeletedElementViewModels.ContainsByPredicate([](TSharedPtr<FHierarchyElementViewModel> Candidate)
+			{
+				return Candidate->GetData()->IsA<UHierarchyItem>() || Candidate->GetData()->IsA<UHierarchyCategory>();
+			});
+			
+			bool bContainsSection = DeletedElementViewModels.ContainsByPredicate([](TSharedPtr<FHierarchyElementViewModel> Candidate)
+			{
+				return Candidate->GetData()->IsA<UHierarchySection>();
+			});
+			
+			bool bContainsActiveSection = DeletedElementViewModels.Contains(ActiveHierarchySectionViewModel.Pin());
+			
+			
+			if(bContainsHierarchyElement)
+			{
+				RefreshHierarchyView(false);
+			}
+			else if(bContainsSection)
+			{
+				RefreshSectionsView();
+
+				if(bContainsActiveSection)
+				{
+					SetActiveHierarchySection(GetDefaultHierarchySectionViewModel());
+				}
+			}
+		}
+	}
+	
+	OnHierarchyChangedDelegate.Broadcast(Payload);
+}
+
 TSubclassOf<UHierarchyCategory> UDataHierarchyViewModelBase::GetCategoryDataClass() const
 {
 	return UHierarchyCategory::StaticClass();
@@ -632,7 +713,7 @@ bool UDataHierarchyViewModelBase::ToolMenuCanRequestRename(const FToolMenuContex
 	return false;
 }
 
-void UDataHierarchyViewModelBase::ToolMenuDelete(const FToolMenuContext& Context) const
+void UDataHierarchyViewModelBase::ToolMenuDelete(const FToolMenuContext& Context)
 {
 	UHierarchyMenuContext* HierarchyMenuContext = Context.FindContext<UHierarchyMenuContext>();
 	DeleteElements(HierarchyMenuContext->MenuHierarchyElements);
@@ -756,11 +837,10 @@ TSharedPtr<FHierarchyElementViewModel> FHierarchyElementViewModel::AddChild(TSub
 	TSharedPtr<FHierarchyElementViewModel> ElementViewModel = FindViewModelForChild(NewChild);
 	ensureMsgf(ElementViewModel.IsValid(), TEXT("Invalid view model detected after syncing view models to data."));
 	
-	TInstancedStruct<FHierarchyElementChangedPayload_AddedElement> Payload = TInstancedStruct<FHierarchyElementChangedPayload_AddedElement>::Make();
-	Payload.GetMutable<FHierarchyElementChangedPayload_AddedElement>().AddedElementViewModel = ElementViewModel;
-	HierarchyViewModel->OnHierarchyElementChanged().ExecuteIfBound(Payload);
+	TInstancedStruct<FHierarchyChangedPayload_ElementsAdded> Payload = TInstancedStruct<FHierarchyChangedPayload_ElementsAdded>::Make();
+	Payload.GetMutable<FHierarchyChangedPayload_ElementsAdded>().AddedElementViewModels.Add(ElementViewModel);
+	HierarchyViewModel->OnHierarchyChanged(Payload);
 	
-	HierarchyViewModel->OnHierarchyChanged().Broadcast();
 	return ElementViewModel;
 }
 
@@ -942,9 +1022,13 @@ TSharedRef<FHierarchyElementViewModel> FHierarchyElementViewModel::DuplicateToTh
 	}
 	
 	SyncViewModelsToData();
-
-	HierarchyViewModel->OnHierarchyChanged().Broadcast();
+	
 	TSharedPtr<FHierarchyElementViewModel> ViewModel = FindViewModelForChild(NewItem);
+	
+	TInstancedStruct<FHierarchyChangedPayload_ElementsAdded> Payload = TInstancedStruct<FHierarchyChangedPayload_ElementsAdded>::Make();
+	Payload.GetMutablePtr<FHierarchyChangedPayload_ElementsAdded>()->AddedElementViewModels.Add(ViewModel);
+	HierarchyViewModel->OnHierarchyChanged(Payload);
+	
 	return ViewModel.ToSharedRef();
 }
 
@@ -961,9 +1045,15 @@ TSharedRef<FHierarchyElementViewModel> FHierarchyElementViewModel::ReparentToThi
 	}
 	
 	ItemToMove->Delete();
+	
 	SyncViewModelsToData();
-	HierarchyViewModel->OnHierarchyChanged().Broadcast();
+	
 	TSharedPtr<FHierarchyElementViewModel> ViewModel = FindViewModelForChild(NewItem);
+	
+	TInstancedStruct<FHierarchyChangedPayload_ElementsMoved> Payload = TInstancedStruct<FHierarchyChangedPayload_ElementsMoved>::Make();
+	Payload.GetMutablePtr<FHierarchyChangedPayload_ElementsMoved>()->MovedElementViewModels.Add(ViewModel);
+	HierarchyViewModel->OnHierarchyChanged(Payload);
+	
 	return ViewModel.ToSharedRef();
 }
 
@@ -1017,6 +1107,42 @@ TSharedPtr<FHierarchyElementViewModel> FHierarchyElementViewModel::FindViewModel
 	return nullptr;
 }
 
+TSharedPtr<FHierarchyElementViewModel> FHierarchyElementViewModel::FindViewModelForObject(const UObject* Object, bool bSearchRecursively) const
+{
+	// If it's a hierarchy element we use the dedicated function
+	if (Object->IsA<UHierarchyElement>())
+	{
+		return FindViewModelForChild(Cast<UHierarchyElement>(Object), bSearchRecursively);
+	}
+	
+	// If not, we are looking for the view model that points to the object for editing
+	for (TSharedPtr<FHierarchyElementViewModel> ChildViewModel : Children)
+	{
+		if (UObject* ObjectForEditing = ChildViewModel->GetObjectForEditing())
+		{
+			if (ObjectForEditing == Object)
+			{
+				return ChildViewModel;
+			}
+		}
+	}
+	
+	if (bSearchRecursively)
+	{
+		for(TSharedPtr<FHierarchyElementViewModel> ChildViewModel : Children)
+		{
+			TSharedPtr<FHierarchyElementViewModel> FoundViewModel = ChildViewModel->FindViewModelForObject(Object, bSearchRecursively);
+
+			if(FoundViewModel.IsValid())
+			{
+				return FoundViewModel;
+			}
+		}
+	}
+	
+	return nullptr;
+}
+
 int32 FHierarchyElementViewModel::FindIndexOfChild(const UHierarchyElement* Child) const
 {
 	return Children.FindLastByPredicate([Child](TSharedPtr<FHierarchyElementViewModel> Item)
@@ -1047,9 +1173,9 @@ void FHierarchyElementViewModel::DeleteChild(TSharedPtr<FHierarchyElementViewMod
 	GetDataMutable()->GetChildrenMutable().Remove(Child->GetDataMutable());
 	Children.Remove(Child);
 
-	TInstancedStruct<FHierarchyElementChangedPayload_DeletedElement> Payload = TInstancedStruct<FHierarchyElementChangedPayload_DeletedElement>::Make();
-	Payload.GetMutable<FHierarchyElementChangedPayload_DeletedElement>().DeletedElementViewModel = Child;
-	HierarchyViewModel->OnHierarchyElementChanged().ExecuteIfBound(Payload);
+	TInstancedStruct<FHierarchyChangedPayload_ElementsDeleted> Payload = TInstancedStruct<FHierarchyChangedPayload_ElementsDeleted>::Make();
+	Payload.GetMutable<FHierarchyChangedPayload_ElementsDeleted>().DeletedElementViewModels.Add(Child);
+	HierarchyViewModel->OnHierarchyChanged(Payload);
 }
 
 TOptional<EItemDropZone> FHierarchyElementViewModel::OnCanRowAcceptDrop(const FDragDropEvent& DragDropEvent, EItemDropZone ItemDropZone, TSharedPtr<FHierarchyElementViewModel> Item)
@@ -1097,6 +1223,20 @@ FHierarchyElementViewModel::FResultWithUserFeedback FHierarchyElementViewModel::
 	return CanDragInternal();
 }
 
+void FHierarchyElementViewModel::Rename(FName NewName)
+{
+	FName OldName = FName(ToString());
+	
+	RenameInternal(NewName); 
+	
+	TInstancedStruct<FHierarchyChangedPayload_ElementNameChanged> Payload = TInstancedStruct<FHierarchyChangedPayload_ElementNameChanged>::Make();
+	Payload.GetMutable<FHierarchyChangedPayload_ElementNameChanged>().ElementViewModel = AsShared();
+	Payload.GetMutable<FHierarchyChangedPayload_ElementNameChanged>().OldName = OldName;
+	Payload.GetMutable<FHierarchyChangedPayload_ElementNameChanged>().NewName = NewName;
+	
+	HierarchyViewModel->OnHierarchyChanged(Payload);
+}
+
 bool FHierarchyElementViewModel::IsExpandedByDefault() const
 {
 	// Otherwise, let the element decide
@@ -1418,9 +1558,9 @@ TSharedPtr<FHierarchySectionViewModel> FHierarchyRootViewModel::AddSection()
 	
 	OnSectionAddedDelegate.ExecuteIfBound(SectionViewModel);
 	
-	TInstancedStruct<FHierarchyElementChangedPayload_AddedElement> Payload = TInstancedStruct<FHierarchyElementChangedPayload_AddedElement>::Make();
-	Payload.GetMutable<FHierarchyElementChangedPayload_AddedElement>().AddedElementViewModel = SectionViewModel;
-	HierarchyViewModel->OnHierarchyElementChanged().ExecuteIfBound(Payload);
+	TInstancedStruct<FHierarchyChangedPayload_ElementsAdded> Payload = TInstancedStruct<FHierarchyChangedPayload_ElementsAdded>::Make();
+	Payload.GetMutable<FHierarchyChangedPayload_ElementsAdded>().AddedElementViewModels.Add(SectionViewModel);
+	HierarchyViewModel->OnHierarchyChanged(Payload);
 	
 	OnSectionsChangedDelegate.ExecuteIfBound();
 	
@@ -1435,9 +1575,9 @@ void FHierarchyRootViewModel::DeleteSection(TSharedPtr<FHierarchyElementViewMode
 
 	OnSectionDeletedDelegate.ExecuteIfBound(SectionViewModel);
 	
-	TInstancedStruct<FHierarchyElementChangedPayload_DeletedElement> Payload = TInstancedStruct<FHierarchyElementChangedPayload_DeletedElement>::Make();
-	Payload.GetMutable<FHierarchyElementChangedPayload_DeletedElement>().DeletedElementViewModel = SectionViewModel;
-	HierarchyViewModel->OnHierarchyElementChanged().ExecuteIfBound(Payload);
+	TInstancedStruct<FHierarchyChangedPayload_ElementsDeleted> Payload = TInstancedStruct<FHierarchyChangedPayload_ElementsDeleted>::Make();
+	Payload.GetMutable<FHierarchyChangedPayload_ElementsDeleted>().DeletedElementViewModels.Add(SectionViewModel);
+	HierarchyViewModel->OnHierarchyChanged(Payload);
 	
 	OnSectionsChangedDelegate.ExecuteIfBound();
 }
@@ -1641,7 +1781,7 @@ TOptional<bool> FHierarchySectionViewModel::OnDroppedOnOverride(TSharedPtr<FHier
 		if(bDropSucceeded)
 		{
 			HierarchyViewModel->ForceFullRefresh();
-			HierarchyViewModel->OnHierarchyChanged().Broadcast();
+			HierarchyViewModel->OnHierarchyChanged(TInstancedStruct<FHierarchyChangedPayload>::Make("Dropped section next to section"));
 		}
 		
 		return true;
@@ -1671,7 +1811,7 @@ TOptional<bool> FHierarchySectionViewModel::OnDroppedOnOverride(TSharedPtr<FHier
 		}
 		
 		HierarchyViewModel->RefreshHierarchyView();
-		HierarchyViewModel->OnHierarchyChanged().Broadcast();
+		HierarchyViewModel->OnHierarchyChanged(TInstancedStruct<FHierarchyChangedPayload>::Make("Dropped element on section"));
 		
 		return true;
 	}
@@ -1792,7 +1932,7 @@ const FTableRowStyle* FHierarchyCategoryViewModel::GetRowStyle() const
 	return &FDataHierarchyEditorStyle::Get().GetWidgetStyle<FTableRowStyle>("HierarchyEditor.Row.Category");
 }
 
-void UDataHierarchyViewModelBase::AddCategory(TSharedPtr<FHierarchyElementViewModel> CategoryParent) const
+void UDataHierarchyViewModelBase::AddCategory(TSharedPtr<FHierarchyElementViewModel> CategoryParent)
 {
 	// If no category parent was specified, we add it to the root
 	if(CategoryParent == nullptr)
@@ -1830,9 +1970,9 @@ void UDataHierarchyViewModelBase::AddCategory(TSharedPtr<FHierarchyElementViewMo
 
 		CategoryViewModel->SetCategoryName(UE::DataHierarchyEditor::GetUniqueName(FName("New Category"), CategoryNames));
 		
-		RefreshHierarchyView();
-		
-		//OnElementAddedDelegate.ExecuteIfBound(ViewModel);
+		TInstancedStruct<FHierarchyChangedPayload_ElementsAdded> Payload = TInstancedStruct<FHierarchyChangedPayload_ElementsAdded>::Make();
+		Payload.GetMutable<FHierarchyChangedPayload_ElementsAdded>().AddedElementViewModels.Add(CategoryViewModel);
+		OnHierarchyChanged(Payload);
 	}
 }
 
@@ -1854,7 +1994,7 @@ void UDataHierarchyViewModelBase::GenerateDynamicContextMenu(UToolMenu* ToolMenu
 	}
 }
 
-void UDataHierarchyViewModelBase::GenerateDynamicContextMenuInternal(UToolMenu* DynamicToolMenu) const
+void UDataHierarchyViewModelBase::GenerateDynamicContextMenuInternal(UToolMenu* DynamicToolMenu)
 {
 	UHierarchyMenuContext* HierarchyMenuContext = DynamicToolMenu->FindContext<UHierarchyMenuContext>();
 
@@ -1925,32 +2065,38 @@ void UDataHierarchyViewModelBase::DeleteElementWithIdentity(FHierarchyElementIde
 	FScopedTransaction Transaction(LOCTEXT("Transaction_DeleteItem", "Deleted hierarchy item"));
 	HierarchyRoot->Modify();
 
-	bool bItemDeleted = false;
+	TSharedPtr<FHierarchyElementViewModel> ViewModelToDelete;
+	
 	if(TSharedPtr<FHierarchyElementViewModel> ViewModel = HierarchyRootViewModel->FindViewModelForChild(Identity, true))
 	{
 		if(ViewModel->CanDelete())
 		{
-			ViewModel->Delete();
-			bItemDeleted = true;
+			ViewModelToDelete = ViewModel;
 		}
 	}
 	
-	TArray<TSharedPtr<FHierarchySectionViewModel>> SectionViewModels; 
-	HierarchyRootViewModel->GetChildrenViewModelsForType<UHierarchySection, FHierarchySectionViewModel>(SectionViewModels);
-	
-	for(TSharedPtr<FHierarchySectionViewModel> SectionViewModel : SectionViewModels)
+	// If we haven't found an element in the hierarchy, check the sections
+	if (!ViewModelToDelete.IsValid())
 	{
-		if(SectionViewModel->GetData()->GetPersistentIdentity() == Identity && SectionViewModel->CanDelete())
+		for(TSharedPtr<FHierarchySectionViewModel> SectionViewModel : HierarchyRootViewModel->GetSectionViewModels())
 		{
-			SectionViewModel->Delete();
-			bItemDeleted = true;
+			if(SectionViewModel->GetData()->GetPersistentIdentity() == Identity && SectionViewModel->CanDelete())
+			{
+				ViewModelToDelete = SectionViewModel;
+			}
 		}
 	}
 
-	if(bItemDeleted)
+	if (ViewModelToDelete.IsValid())
 	{
+		ViewModelToDelete->Delete();
+		
+		TInstancedStruct<FHierarchyChangedPayload_ElementsDeleted> Payload = TInstancedStruct<FHierarchyChangedPayload_ElementsDeleted>::Make();
+		Payload.GetMutable<FHierarchyChangedPayload_ElementsDeleted>().DeletedElementViewModels.Add(ViewModelToDelete);
+		
 		HierarchyRootViewModel->SyncViewModelsToData();
-		OnHierarchyChangedDelegate.Broadcast();
+		
+		OnHierarchyChanged(Payload);
 	}
 	else
 	{
@@ -1958,25 +2104,36 @@ void UDataHierarchyViewModelBase::DeleteElementWithIdentity(FHierarchyElementIde
 	}
 }
 
-void UDataHierarchyViewModelBase::DeleteElements(TArray<TSharedPtr<FHierarchyElementViewModel>> ViewModels) const
+void UDataHierarchyViewModelBase::DeleteElements(TArray<TSharedPtr<FHierarchyElementViewModel>> ViewModels)
 {
 	FScopedTransaction Transaction(LOCTEXT("Transaction_DeleteHierarchyElements", "Deleted hierarchy elements"));
 	HierarchyRoot->Modify();
 	
 	bool bAnyItemsDeleted = false;
+	
+	TArray<TSharedPtr<FHierarchyElementViewModel>> ViewModelsToDelete;
 	for(TSharedPtr<FHierarchyElementViewModel> ViewModel : ViewModels)
 	{
 		if(ViewModel->CanDelete())
 		{
-			ViewModel->Delete();
-			bAnyItemsDeleted = true;
+			ViewModelsToDelete.Add(ViewModel);
+			
 		}
 	}
+	
+	for (TSharedPtr<FHierarchyElementViewModel> ViewModelToDelete : ViewModelsToDelete)
+	{
+		ViewModelToDelete->Delete();
+		bAnyItemsDeleted = true;
+	}
 
 	if(bAnyItemsDeleted)
 	{
 		HierarchyRootViewModel->SyncViewModelsToData();
-		OnHierarchyChangedDelegate.Broadcast();
+		
+		TInstancedStruct<FHierarchyChangedPayload_ElementsDeleted> Payload = TInstancedStruct<FHierarchyChangedPayload_ElementsDeleted>::Make();
+		Payload.GetMutablePtr<FHierarchyChangedPayload_ElementsDeleted>()->DeletedElementViewModels = ViewModelsToDelete;
+		OnHierarchyChanged(Payload);
 	}
 	else
 	{
diff --git a/Engine/Source/Editor/DataHierarchyEditor/Private/Widgets/SDataHierarchyEditor.cpp b/Engine/Source/Editor/DataHierarchyEditor/Private/Widgets/SDataHierarchyEditor.cpp
--- a/Engine/Source/Editor/DataHierarchyEditor/Private/Widgets/SDataHierarchyEditor.cpp
+++ b/Engine/Source/Editor/DataHierarchyEditor/Private/Widgets/SDataHierarchyEditor.cpp
@@ -539,7 +539,6 @@ void SDataHierarchyEditor::Construct(const FArguments& InArgs, TObjectPtr<UDataH
 	HierarchyViewModel->OnRefreshHierarchyView().BindSP(this, &SDataHierarchyEditor::RefreshHierarchyView);
 	HierarchyViewModel->OnRefreshSectionsView().BindSP(this, &SDataHierarchyEditor::RefreshSectionsView);
 	HierarchyViewModel->OnHierarchySectionActivated().BindSP(this, &SDataHierarchyEditor::OnHierarchySectionActivated);
-	HierarchyViewModel->OnHierarchyElementChanged().BindSP(this, &SDataHierarchyEditor::OnHierarchyElementChanged);
 
 	BindToHierarchyRootViewModel();
 	
@@ -855,7 +854,6 @@ SDataHierarchyEditor::~SDataHierarchyEditor()
 		HierarchyViewModel->OnRefreshHierarchyView().Unbind();
 		HierarchyViewModel->OnRefreshSectionsView().Unbind();
 		HierarchyViewModel->OnHierarchySectionActivated().Unbind();
-		HierarchyViewModel->OnHierarchyElementChanged().Unbind();
 
 		UnbindFromHierarchyRootViewModel();
 
@@ -1462,54 +1460,6 @@ bool SDataHierarchyEditor::CanNavigateToMatchingHierarchyElementFromSelectedSour
 	return false;
 }
 
-void SDataHierarchyEditor::OnHierarchyElementChanged(TInstancedStruct<FHierarchyElementChangedPayload> Payload)
-{
-	if(Payload.IsValid())
-	{
-		if(Payload.GetScriptStruct() == FHierarchyElementChangedPayload_AddedElement::StaticStruct())
-		{
-			TSharedPtr<FHierarchyElementViewModel> AddedElementViewModel = Payload.GetMutablePtr<FHierarchyElementChangedPayload_AddedElement>()->AddedElementViewModel.Pin();
-			
-			// when a new item is created (opposed to dragged & dropped from source view, i.e. only categories so far)
-			// we make sure to request a tree refresh, select the row, and request a pending rename since the widget will created a frame later
-			if(AddedElementViewModel->GetData()->IsA<UHierarchyItem>() || AddedElementViewModel->GetData()->IsA<UHierarchyCategory>())
-			{
-				HierarchyTreeView->RequestTreeRefresh();
-				NavigateToHierarchyElement(AddedElementViewModel);
-			}
-			else if(AddedElementViewModel->GetData()->IsA<UHierarchySection>())
-			{
-				RefreshSectionsView();
-				HierarchyViewModel->SetActiveHierarchySection(StaticCastSharedPtr<FHierarchySectionViewModel>(AddedElementViewModel));
-			}
-
-			AddedElementViewModel->RequestRenamePending();
-		}
-
-		if(Payload.GetScriptStruct() == FHierarchyElementChangedPayload_DeletedElement::StaticStruct())
-		{
-			TSharedPtr<FHierarchyElementViewModel> DeletedElementViewModel = Payload.GetMutablePtr<FHierarchyElementChangedPayload_DeletedElement>()->DeletedElementViewModel.Pin();
-			
-			// when a new item is created (opposed to dragged & dropped from source view, i.e. only categories so far)
-			// we make sure to request a tree refresh, select the row, and request a pending rename since the widget will created a frame later
-			if(DeletedElementViewModel->GetData()->IsA<UHierarchyItem>() || DeletedElementViewModel->GetData()->IsA<UHierarchyCategory>())
-			{
-				HierarchyTreeView->RequestTreeRefresh();
-				NavigateToHierarchyElement(DeletedElementViewModel);
-			}
-			else if(DeletedElementViewModel->GetData()->IsA<UHierarchySection>())
-			{
-				RefreshSectionsView();
-
-				if(HierarchyViewModel->GetActiveHierarchySectionViewModel() == DeletedElementViewModel)
-				{
-					HierarchyViewModel->SetActiveHierarchySection(HierarchyViewModel->GetDefaultHierarchySectionViewModel());
-				}
-			}
-		}
-	}
-}
-
 void SDataHierarchyEditor::OnHierarchySectionActivated(TSharedPtr<FHierarchySectionViewModel> Section)
 {
 	// We forward nullptr in case this is the 'All' default section. We determine this by checking its data for validity.
@@ -2014,8 +1964,8 @@ void SDataHierarchyEditor::NotifyPostChange(const FPropertyChangedEvent& Propert
 {
 	const UObject* EditedObject = PropertyChangedEvent.GetObjectBeingEdited(0);
 
+	TSharedPtr<FHierarchyElementViewModel> MatchingViewModel;
 	// We look for the view model corresponding to the edited object.
-	// The object doesn't have to be a hierarchy element, since view models can provide external elements for editing. 
 	if(EditedObject->IsA<UHierarchySection>())
 	{
 		if(SourceRootViewModel.IsValid())
@@ -2027,7 +1977,7 @@ void SDataHierarchyEditor::NotifyPostChange(const FPropertyChangedEvent& Propert
 			
 			if(FoundSourceSectionViewModel)
 			{
-				(*FoundSourceSectionViewModel)->PostEditChange(PropertyChangedEvent, PropertyThatChanged);
+				MatchingViewModel = *FoundSourceSectionViewModel;
 			}
 		}
 
@@ -2038,26 +1988,37 @@ void SDataHierarchyEditor::NotifyPostChange(const FPropertyChangedEvent& Propert
 
 		if(FoundHierarchySectionViewModel)
 		{
-			(*FoundHierarchySectionViewModel)->PostEditChange(PropertyChangedEvent, PropertyThatChanged);
+			MatchingViewModel = *FoundHierarchySectionViewModel;
 		}
 	}
-	else if(EditedObject->IsA<UHierarchyElement>())
+	// The object doesn't have to be a hierarchy element, since view models can provide external elements for editing.
+	else
 	{
 		if(SourceRootViewModel.IsValid())
 		{
-			if(TSharedPtr<FHierarchyElementViewModel> FoundSourceViewModel = SourceRootViewModel->FindViewModelForChild(Cast<UHierarchyElement>(PropertyChangedEvent.GetObjectBeingEdited(0)), true))
+			if(TSharedPtr<FHierarchyElementViewModel> FoundSourceViewModel = SourceRootViewModel->FindViewModelForObject(PropertyChangedEvent.GetObjectBeingEdited(0), true))
 			{
-				FoundSourceViewModel->PostEditChange(PropertyChangedEvent, PropertyThatChanged);
+				MatchingViewModel = FoundSourceViewModel;
 			}
 		}
 		
-		if(TSharedPtr<FHierarchyElementViewModel> FoundHierarchyViewModel = HierarchyViewModel->GetHierarchyRootViewModel()->FindViewModelForChild(Cast<UHierarchyElement>(PropertyChangedEvent.GetObjectBeingEdited(0)), true))
+		if(TSharedPtr<FHierarchyElementViewModel> FoundHierarchyViewModel = HierarchyViewModel->GetHierarchyRootViewModel()->FindViewModelForObject(PropertyChangedEvent.GetObjectBeingEdited(0), true))
 		{
-			FoundHierarchyViewModel->PostEditChange(PropertyChangedEvent, PropertyThatChanged);
+			MatchingViewModel = FoundHierarchyViewModel;
 		}
 	}
 	
-	HierarchyViewModel->OnHierarchyPropertiesChanged().Broadcast();
+	if (MatchingViewModel.IsValid())
+	{
+		MatchingViewModel->PostEditChange(PropertyChangedEvent, PropertyThatChanged);
+		
+		TInstancedStruct<FHierarchyChangedPayload_ElementPropertyChanged> Payload = TInstancedStruct<FHierarchyChangedPayload_ElementPropertyChanged>::Make();
+		Payload.GetMutable<FHierarchyChangedPayload_ElementPropertyChanged>().ElementViewModel = MatchingViewModel;
+		Payload.GetMutable<FHierarchyChangedPayload_ElementPropertyChanged>().PropertyChangedEvent = &PropertyChangedEvent;
+		Payload.GetMutable<FHierarchyChangedPayload_ElementPropertyChanged>().PropertyThatChanged = PropertyThatChanged;
+
+		HierarchyViewModel->OnHierarchyChanged(Payload);
+	}
 }
 
 #undef LOCTEXT_NAMESPACE
diff --git a/Engine/Source/Editor/DataHierarchyEditor/Public/DataHierarchyViewModelBase.h b/Engine/Source/Editor/DataHierarchyEditor/Public/DataHierarchyViewModelBase.h
--- a/Engine/Source/Editor/DataHierarchyEditor/Public/DataHierarchyViewModelBase.h
+++ b/Engine/Source/Editor/DataHierarchyEditor/Public/DataHierarchyViewModelBase.h
@@ -106,6 +106,9 @@ public:
 	template<class PREDICATE_CLASS>
 	void SortChildren(const PREDICATE_CLASS& Predicate, bool bRecursive = false);
 	
+	template<class PREDICATE_CLASS>
+	void StableSortChildren(const PREDICATE_CLASS& Predicate, bool bRecursive = false);
+	
 	virtual FString ToString() const { return GetName(); }
 	FText ToText() const { return FText::FromString(ToString()); }
 	
@@ -243,6 +246,20 @@ void UHierarchyElement::SortChildren(const PREDICATE_CLASS& Predicate, bool bRec
 	}
 }
 
+template <class PREDICATE_CLASS>
+void UHierarchyElement::StableSortChildren(const PREDICATE_CLASS& Predicate, bool bRecursive)
+{
+	Children.StableSort(Predicate);
+
+	if(bRecursive)
+	{
+		for(TObjectPtr<UHierarchyElement> ChildElement : Children)
+		{
+			ChildElement->StableSortChildren(Predicate, bRecursive);
+		}
+	}
+}
+
 template <typename T>
 bool UHierarchyElement::HasMetaDataOfType()
 {
@@ -469,33 +486,59 @@ public:
 
 /** Generic payload struct for when an element in the hierarchy has changed. */
 USTRUCT()
-struct FHierarchyElementChangedPayload
+struct FHierarchyChangedPayload
+{
+	GENERATED_BODY()
+	
+	FHierarchyChangedPayload() = default;
+	FHierarchyChangedPayload(FString InMessage) : Message(InMessage) {}
+	
+	/** An optional message that can be set. */
+	FString Message;
+};
+
+USTRUCT()
+struct FHierarchyChangedPayload_ElementsAdded : public FHierarchyChangedPayload
 {
 	GENERATED_BODY()
+	
+	TArray<TSharedPtr<FHierarchyElementViewModel>> AddedElementViewModels;
 };
 
 USTRUCT()
-struct FHierarchyElementChangedPayload_AddedElement : public FHierarchyElementChangedPayload
+struct FHierarchyChangedPayload_ElementsMoved : public FHierarchyChangedPayload
 {
 	GENERATED_BODY()
 	
-	TWeakPtr<FHierarchyElementViewModel> AddedElementViewModel;
+	TArray<TSharedPtr<FHierarchyElementViewModel>> MovedElementViewModels;
 };
 
 USTRUCT()
-struct FHierarchyElementChangedPayload_DeletedElement : public FHierarchyElementChangedPayload
+struct FHierarchyChangedPayload_ElementsDeleted : public FHierarchyChangedPayload
 {
 	GENERATED_BODY()
 
-	TWeakPtr<FHierarchyElementViewModel> DeletedElementViewModel;
+	TArray<TSharedPtr<FHierarchyElementViewModel>> DeletedElementViewModels;
 };
 
 USTRUCT()
-struct FHierarchyElementChangedPayload_ElementPropertyChanged : public FHierarchyElementChangedPayload
+struct FHierarchyChangedPayload_ElementPropertyChanged : public FHierarchyChangedPayload
 {
 	GENERATED_BODY()
+	
+	TSharedPtr<FHierarchyElementViewModel> ElementViewModel;
+	const FPropertyChangedEvent* PropertyChangedEvent = nullptr;
+	FProperty* PropertyThatChanged = nullptr;
+};
 
-	TWeakPtr<FHierarchyElementViewModel> ElementViewModel;
+USTRUCT()
+struct FHierarchyChangedPayload_ElementNameChanged : public FHierarchyChangedPayload
+{
+	GENERATED_BODY()
+	
+	TSharedPtr<FHierarchyElementViewModel> ElementViewModel;
+	FName OldName;
+	FName NewName;
 };
 
 /** The main controller class for the SHierarchyEditor widget. Defines core hierarchy rules.
@@ -504,14 +547,11 @@ UCLASS(MinimalAPI, Abstract, Transient)
 class UDataHierarchyViewModelBase : public UObject, public FSelfRegisteringEditorUndoClient
 {
 public:
-	DECLARE_MULTICAST_DELEGATE(FOnHierarchyChanged)
-	DECLARE_MULTICAST_DELEGATE(FOnHierarchyPropertiesChanged)
+	DECLARE_MULTICAST_DELEGATE_OneParam(FOnHierarchyChanged, const TInstancedStruct<FHierarchyChangedPayload>&)
 	DECLARE_DELEGATE_OneParam(FOnSectionActivated, TSharedPtr<FHierarchySectionViewModel> Section)
-	DECLARE_DELEGATE_OneParam(FOnElementAdded, TSharedPtr<FHierarchyElementViewModel> AddedElement)
 	DECLARE_DELEGATE_OneParam(FOnRefreshViewRequested, bool bForceFullRefresh)
 	DECLARE_DELEGATE_OneParam(FOnNavigateToElementIdentityInHierarchyRequested, FHierarchyElementIdentity Identity)
 	DECLARE_DELEGATE_OneParam(FOnNavigateToElementInHierarchyRequested, TSharedPtr<FHierarchyElementViewModel> ElementViewModel)
-	DECLARE_DELEGATE_OneParam(FOnHierarchyElementChanged, TInstancedStruct<FHierarchyElementChangedPayload>)
 
 	GENERATED_BODY()
 
@@ -546,6 +586,9 @@ public:
 	 * The root view model is also given as a way to forcefully sync view models to access additional functionality, if needed. Only used if UsesSourceItems is true. */
 	virtual void PrepareSourceItems(UHierarchyRoot* SourceRoot, TSharedPtr<FHierarchyRootViewModel> SourceRootViewModel) {};
 
+	/** Generic callback for when something in the hierarchy has changed. */
+	UE_API virtual void OnHierarchyChanged(const TInstancedStruct<FHierarchyChangedPayload>& Payload);
+	
 	/** If specified, will add a button in the Data Hierarchy Editor that allows adding the given types from a menu rather than drag & dropping them into the hierarchy. */
 	virtual TArray<TSubclassOf<UHierarchyElement>> GetAdditionalTypesToAddInUi() const { return TArray<TSubclassOf<UHierarchyElement>>(); }
 	
@@ -583,17 +626,17 @@ private:
 	virtual TSharedPtr<FHierarchyElementViewModel> CreateCustomViewModelForElement(UHierarchyElement* Element, TSharedPtr<FHierarchyElementViewModel> Parent) { return nullptr; }
 	/** Override this to add custom context menu options. Defaults to common actions such as finding, renaming and deleting.
 	 * The hierarchy elements to generate this menu for are contained in a UHierarchyMenuContext object accessible through the ToolMenu. */
-	UE_API virtual void GenerateDynamicContextMenuInternal(UToolMenu* DynamicToolMenu) const;
+	UE_API virtual void GenerateDynamicContextMenuInternal(UToolMenu* DynamicToolMenu);
 public:
 	const UHierarchyDataRefreshContext* GetRefreshContext() const { return RefreshContext; }
 	/** Set the refresh context to easily allow Hierarchy Elements representing external data to access whether the external data still exists. */
 	void SetRefreshContext(UHierarchyDataRefreshContext* InContext) { RefreshContext = InContext; }
 	
 	UE_API void AddElementUnderRoot(TSubclassOf<UHierarchyElement> NewChildClass, FHierarchyElementIdentity ChildIdentity);
-	UE_API void AddCategory(TSharedPtr<FHierarchyElementViewModel> CategoryParent = nullptr) const;
+	UE_API void AddCategory(TSharedPtr<FHierarchyElementViewModel> CategoryParent = nullptr);
 
 	/** Delete all specified elements. */
-	UE_API void DeleteElements(TArray<TSharedPtr<FHierarchyElementViewModel>> ViewModels) const;
+	UE_API void DeleteElements(TArray<TSharedPtr<FHierarchyElementViewModel>> ViewModels);
 	/** Special case for deleting a specific element based on its identity.
 	 *  Useful for externally removing an element from the hierarchy when you don't have access to the view model.*/
 	UE_API void DeleteElementWithIdentity(FHierarchyElementIdentity Identity);
@@ -622,10 +665,7 @@ public:
 	FSimpleDelegate& OnRefreshSectionsView() { return RefreshSectionsViewDelegate; }
 
 	// Delegates for external systems
-	FOnHierarchyElementChanged& OnHierarchyElementChanged() { return OnHierarchyElementChangedDelegate; }
 	FOnHierarchyChanged& OnHierarchyChanged() { return OnHierarchyChangedDelegate; } 
-	FOnHierarchyChanged& OnHierarchyPropertiesChanged() { return OnHierarchyPropertiesChangedDelegate; } 
-	//FOnElementAdded& OnElementAdded() { return OnElementAddedDelegate; }
 	FOnRefreshViewRequested& OnRefreshViewRequested() { return RefreshAllViewsRequestedDelegate; }
 	FOnNavigateToElementIdentityInHierarchyRequested& OnNavigateToElementIdentityInHierarchyRequested() { return OnNavigateToElementIdentityInHierarchyRequestedDelegate; }
 	FOnNavigateToElementInHierarchyRequested& OnNavigateToElementInHierarchyRequested() { return OnNavigateToElementInHierarchyRequestedDelegate; }
@@ -652,7 +692,7 @@ private:
 	UE_API void ToolMenuRequestRename(const FToolMenuContext& Context) const;
 	UE_API bool ToolMenuCanRequestRename(const FToolMenuContext& Context) const;
 
-	UE_API void ToolMenuDelete(const FToolMenuContext& Context) const;
+	UE_API void ToolMenuDelete(const FToolMenuContext& Context);
 	UE_API bool ToolMenuCanDelete(const FToolMenuContext& Context) const;
 
 	UE_API void ToolMenuNavigateTo(const FToolMenuContext& Context) const;
@@ -682,13 +722,9 @@ protected:
 	FOnNavigateToElementIdentityInHierarchyRequested OnNavigateToElementIdentityInHierarchyRequestedDelegate;
 	FOnNavigateToElementInHierarchyRequested OnNavigateToElementInHierarchyRequestedDelegate;
 	
-	//FOnElementAdded OnElementAddedDelegate;
 	FOnSectionActivated OnHierarchySectionActivatedDelegate;
 	FOnSectionActivated OnSourceSectionActivatedDelegate;
 	FOnHierarchyChanged OnHierarchyChangedDelegate;
-	FOnHierarchyElementChanged OnHierarchyElementChangedDelegate;
-
-	FOnHierarchyPropertiesChanged OnHierarchyPropertiesChangedDelegate;
 	
 	FSimpleDelegate OnInitializedDelegate;
 
@@ -771,8 +807,13 @@ struct FHierarchyElementViewModel : TSharedFromThis<FHierarchyElementViewModel>,
 	UE_API TSharedRef<FHierarchyElementViewModel> DuplicateToThis(TSharedPtr<FHierarchyElementViewModel> ItemToDuplicate, int32 InsertIndex = INDEX_NONE);
 	UE_API TSharedRef<FHierarchyElementViewModel> ReparentToThis(TSharedPtr<FHierarchyElementViewModel> ItemToMove, int32 InsertIndex = INDEX_NONE);
 
+	/** Finds the view model for a given hierarchy element. Needs to be a child of the current element somewhere down the hierarchy. */
 	UE_API TSharedPtr<FHierarchyElementViewModel> FindViewModelForChild(const UHierarchyElement* Child, bool bSearchRecursively = false) const;
+	/** Finds the view model for a given hierarchy element identity. Needs to be a child of the current element somewhere down the hierarchy. */
 	UE_API TSharedPtr<FHierarchyElementViewModel> FindViewModelForChild(FHierarchyElementIdentity ChildIdentity, bool bSearchRecursively = false) const;
+	/** Finds the view model for a given object. Can be a hierarchy element, but view models can also represent other objects not part of the hierarchy. */
+	UE_API TSharedPtr<FHierarchyElementViewModel> FindViewModelForObject(const UObject* Object, bool bSearchRecursively = false) const;
+	
 	UE_API int32 FindIndexOfChild(const UHierarchyElement* Child) const;
 	UE_API int32 FindIndexOfDataChild(TSharedPtr<const FHierarchyElementViewModel> Child) const;
 	UE_API int32 FindIndexOfDataChild(const UHierarchyElement* Child) const;
@@ -824,7 +865,7 @@ struct FHierarchyElementViewModel : TSharedFromThis<FHierarchyElementViewModel>,
 	/** Returns true if renamable */
 	bool CanRename() { return IsEditableByUser().bResult && CanRenameInternal(); }
 
-	void Rename(FName NewName) { RenameInternal(NewName); HierarchyViewModel->OnHierarchyPropertiesChanged().Broadcast(); }
+	void Rename(FName NewName);
 
 	void RequestRename()
 	{
