Fixed section association for new elements when new elements are added to the hierarchy. New categories in a non-default section will get correctly section'ized

Fixed All sections being hardcoded to be UHierarchySections instead of the user defined TSubclassOf<UHierarchySection>
> This fixes an issue where in Tagged Asset Browser Configuration the custom section view model would not be created allowing custom drag & drop on the All section

Fixed hierarchy categories not displaying their tooltips in the hierarchy itself

#Synced-CL 49840944

diff --git a/Engine/Source/Editor/DataHierarchyEditor/Internal/Widgets/SDataHierarchyEditor.h b/Engine/Source/Editor/DataHierarchyEditor/Internal/Widgets/SDataHierarchyEditor.h
--- a/Engine/Source/Editor/DataHierarchyEditor/Internal/Widgets/SDataHierarchyEditor.h
+++ b/Engine/Source/Editor/DataHierarchyEditor/Internal/Widgets/SDataHierarchyEditor.h
@@ -271,6 +271,7 @@ public:
 	UE_API FText GetElementText() const;
 	UE_API void OnRenameElement(const FText& NewText, ETextCommit::Type) const;
 	
+	FText GetTooltipText() const;
 private:
 	TWeakPtr<struct FHierarchyElementViewModel> ElementViewModelWeak;
 	TSharedPtr<SInlineEditableTextBlock> InlineEditableTextBlock;
diff --git a/Engine/Source/Editor/DataHierarchyEditor/Private/DataHierarchyViewModelBase.cpp b/Engine/Source/Editor/DataHierarchyEditor/Private/DataHierarchyViewModelBase.cpp
--- a/Engine/Source/Editor/DataHierarchyEditor/Private/DataHierarchyViewModelBase.cpp
+++ b/Engine/Source/Editor/DataHierarchyEditor/Private/DataHierarchyViewModelBase.cpp
@@ -457,9 +457,9 @@ void UDataHierarchyViewModelBase::Initialize()
 	HierarchyRootViewModel->Initialize();
 	HierarchyRootViewModel->AddChildFilter(FHierarchyElementViewModel::FOnFilterChild::CreateUObject(this, &UDataHierarchyViewModelBase::FilterForHierarchySection));
 	HierarchyRootViewModel->SyncViewModelsToData();
-
-	AllSection = NewObject<UHierarchySection>(this, AllSectionHierarchyObjectName, RF_Transient);
-	DefaultHierarchySectionViewModel = MakeShared<FHierarchySectionViewModel>(AllSection, GetHierarchyRootViewModel().ToSharedRef(), this);
+	
+	AllSection = NewObject<UHierarchySection>(this, GetSectionDataClass(), AllSectionHierarchyObjectName,  RF_Transient);
+	DefaultHierarchySectionViewModel = StaticCastSharedPtr<FHierarchySectionViewModel>(CreateViewModelForElement(AllSection, GetHierarchyRootViewModel().ToSharedRef()));
 	DefaultHierarchySectionViewModel->SetSectionName(FName(AllSectionDisplayName.ToString()));
 	SetActiveHierarchySection(DefaultHierarchySectionViewModel);
 	
@@ -535,7 +535,7 @@ void UDataHierarchyViewModelBase::OnHierarchyChanged(const TInstancedStruct<FHie
 			
 			// when a new item is created (opposed to dragged & dropped from source view, i.e. only categories so far)
 			// we make sure to request a tree refresh, select the row, and request a pending rename since the widget will be created a frame later
-			if(bContainsHierarchyElement)
+			if (bContainsHierarchyElement)
 			{
 				RefreshHierarchyView(false);
 				
@@ -544,7 +544,8 @@ void UDataHierarchyViewModelBase::OnHierarchyChanged(const TInstancedStruct<FHie
 					NavigateToElementInHierarchy(AddedElementViewModels[0].ToSharedRef());
 				}
 			}
-			else if(bContainsSection)
+			
+			if (bContainsSection)
 			{
 				RefreshSectionsView();
 				
@@ -553,6 +554,37 @@ void UDataHierarchyViewModelBase::OnHierarchyChanged(const TInstancedStruct<FHie
 					SetActiveHierarchySection(StaticCastSharedPtr<FHierarchySectionViewModel>(AddedElementViewModels[0]));
 				}
 			}
+			
+			// For any element added to the root, we fixup section association
+			TSharedPtr<FHierarchySectionViewModel> ActiveSectionViewModel = GetActiveHierarchySectionViewModel();
+			for (TSharedPtr<FHierarchyElementViewModel> AddedElement : AddedElementViewModels)
+			{
+				if (AddedElement->GetData()->IsA<UHierarchySection>())
+				{
+					continue;
+				}
+				
+				TSharedPtr<FHierarchyElementViewModel> ParentViewModel = AddedElement->GetParent().Pin();
+				if (!ParentViewModel->GetData()->IsA<UHierarchyRoot>())
+				{
+					continue;
+				}
+				
+				// If we are in the all section, we clear the section association
+				if (ActiveSectionViewModel->IsAllSection())
+				{
+					AddedElement->GetDataMutable()->DeleteMetaDataOfType<FDataHierarchyElementMetaData_SectionAssociation>();
+				}
+				// If not, we attempt to update the section association
+				else
+				{
+					if (ActiveSectionViewModel->CanContain(AddedElement->GetData()->GetClass()).bResult)
+					{
+						FDataHierarchyElementMetaData_SectionAssociation* SectionAssociation = AddedElement->GetDataMutable()->FindOrAddMetaDataOfType<FDataHierarchyElementMetaData_SectionAssociation>();
+						SectionAssociation->Section = ActiveSectionViewModel->GetSectionData();
+					}
+				}
+			}
 
 			if (AddedElementViewModels.Num() == 1)
 			{
diff --git a/Engine/Source/Editor/DataHierarchyEditor/Private/Widgets/SDataHierarchyEditor.cpp b/Engine/Source/Editor/DataHierarchyEditor/Private/Widgets/SDataHierarchyEditor.cpp
--- a/Engine/Source/Editor/DataHierarchyEditor/Private/Widgets/SDataHierarchyEditor.cpp
+++ b/Engine/Source/Editor/DataHierarchyEditor/Private/Widgets/SDataHierarchyEditor.cpp
@@ -50,6 +50,8 @@ void SHierarchyElement::Construct(const FArguments& InArgs, TSharedPtr<struct FH
 	
 	InViewModel->GetOnRequestRename().BindSP(this, &SHierarchyElement::EnterEditingMode);
 	
+	SetToolTipText(TAttribute<FText>::CreateSP(this, &SHierarchyElement::GetTooltipText));
+	
 	ChildSlot
 	[
 		SAssignNew(InlineEditableTextBlock, SInlineEditableTextBlock)
@@ -108,6 +110,19 @@ void SHierarchyElement::OnRenameElement(const FText& NewText, ETextCommit::Type)
 	}
 }
 
+FText SHierarchyElement::GetTooltipText() const
+{
+	if (ElementViewModelWeak.IsValid())
+	{
+		if (const UHierarchyCategory* Category = Cast<UHierarchyCategory>(ElementViewModelWeak.Pin()->GetData()))
+		{
+			return Category->GetTooltip();
+		}
+	}
+	
+	return FText::GetEmpty();
+}
+
 void SHierarchySection::Construct(const FArguments& InArgs, TSharedPtr<FHierarchySectionViewModel> InSection)
 {
 	SectionViewModelWeak = InSection;
@@ -525,7 +540,7 @@ void SDataHierarchyEditor::Construct(const FArguments& InArgs, TObjectPtr<UDataH
 		SourceRootViewModel->OnSyncPropagated().BindSP(this, &SDataHierarchyEditor::RequestRefreshSourceViewNextFrame, false);
 		SourceRootViewModel->OnSectionsChanged().BindSP(this, &SDataHierarchyEditor::RefreshSectionsView);
 
-		AllSourceSection.Reset(NewObject<UHierarchySection>(HierarchyViewModel.Get(), UDataHierarchyViewModelBase::AllSectionSourceObjectName));
+		AllSourceSection.Reset(NewObject<UHierarchySection>(HierarchyViewModel.Get(), HierarchyViewModel->GetSectionDataClass(), UDataHierarchyViewModelBase::AllSectionSourceObjectName));
 		DefaultSourceSectionViewModel = MakeShared<FHierarchySectionViewModel>(AllSourceSection.Get(), SourceRootViewModel.ToSharedRef(), HierarchyViewModel);
 		DefaultSourceSectionViewModel->SetSectionName("All");
 	}
@@ -1528,15 +1543,7 @@ TSharedRef<SWidget> SDataHierarchyEditor::OnGetCustomAddContent()
 
 					if(ParentViewModel->CanContain(Type).bResult)
 					{
-						TSharedPtr<FHierarchyElementViewModel> NewElementViewModel = ParentViewModel->AddChild(Type);
-
-						if(const UHierarchySection* Section = DataHierarchyEditor->HierarchyViewModel->GetActiveHierarchySectionData())
-						{
-							FDataHierarchyElementMetaData_SectionAssociation* SectionAssociation = NewElementViewModel->GetDataMutable()->FindOrAddMetaDataOfType<FDataHierarchyElementMetaData_SectionAssociation>();
-							SectionAssociation->Section = Section;
-						}
-
-						NewElementViewModel->SyncViewModelsToData();
+						ParentViewModel->AddChild(Type);
 					}
 				});
 				UIAction.CanExecuteAction = FToolMenuCanExecuteAction::CreateLambda([ContextObject, Type](const FToolMenuContext& Context)
