Added PreDeletion event to Hierarchy Editor that is used by new view model function "ShouldReparentChildrenOnDelete". It will attempt to reparent existing children to the next valid parent

#Synced-CL 51048610

diff --git a/Engine/Plugins/FX/Niagara/Source/NiagaraEditor/Private/Customizations/UserParameters/NiagaraComponentUserParametersBuilder.cpp b/Engine/Plugins/FX/Niagara/Source/NiagaraEditor/Private/Customizations/UserParameters/NiagaraComponentUserParametersBuilder.cpp
--- a/Engine/Plugins/FX/Niagara/Source/NiagaraEditor/Private/Customizations/UserParameters/NiagaraComponentUserParametersBuilder.cpp
+++ b/Engine/Plugins/FX/Niagara/Source/NiagaraEditor/Private/Customizations/UserParameters/NiagaraComponentUserParametersBuilder.cpp
@@ -283,7 +283,7 @@ FNiagaraComponentUserParametersNodeBuilder::~FNiagaraComponentUserParametersNode
 			TSharedPtr<FNiagaraSystemViewModel> SystemViewModel = TNiagaraViewModelManager<UNiagaraSystem, FNiagaraSystemViewModel>::GetExistingViewModelForObject(System);
 			if(SystemViewModel.IsValid() && SystemViewModel->GetUserParametersHierarchyViewModel())
 			{
-				SystemViewModel->GetUserParametersHierarchyViewModel()->GetOnHierarchyChangedDelegate().RemoveAll(this);
+				SystemViewModel->GetUserParametersHierarchyViewModel()->GetOnHierarchyEventDelegate().RemoveAll(this);
 			}
 		}
 	}
@@ -500,8 +500,8 @@ void FNiagaraComponentUserParametersNodeBuilder::RegisterRebuildOnHierarchyChang
 			TSharedPtr<FNiagaraSystemViewModel> SystemViewModel = TNiagaraViewModelManager<UNiagaraSystem, FNiagaraSystemViewModel>::GetExistingViewModelForObject(Asset);
 			if(SystemViewModel.IsValid())
 			{
-				SystemViewModel->GetUserParametersHierarchyViewModel()->GetOnHierarchyChangedDelegate().RemoveAll(this);
-				SystemViewModel->GetUserParametersHierarchyViewModel()->GetOnHierarchyChangedDelegate().Add(UDataHierarchyViewModelBase::FOnHierarchyChanged::FDelegate::CreateSP(this, &FNiagaraComponentUserParametersNodeBuilder::RebuildOnHierarchyChanged));
+				SystemViewModel->GetUserParametersHierarchyViewModel()->GetOnHierarchyEventDelegate().RemoveAll(this);
+				SystemViewModel->GetUserParametersHierarchyViewModel()->GetOnHierarchyEventDelegate().Add(UDataHierarchyViewModelBase::FOnHierarchyEvent::FDelegate::CreateSP(this, &FNiagaraComponentUserParametersNodeBuilder::RebuildOnHierarchyChanged));
 			}
 		}
 	}
diff --git a/Engine/Plugins/FX/Niagara/Source/NiagaraEditor/Private/Customizations/UserParameters/NiagaraSystemUserParametersBuilder.cpp b/Engine/Plugins/FX/Niagara/Source/NiagaraEditor/Private/Customizations/UserParameters/NiagaraSystemUserParametersBuilder.cpp
--- a/Engine/Plugins/FX/Niagara/Source/NiagaraEditor/Private/Customizations/UserParameters/NiagaraSystemUserParametersBuilder.cpp
+++ b/Engine/Plugins/FX/Niagara/Source/NiagaraEditor/Private/Customizations/UserParameters/NiagaraSystemUserParametersBuilder.cpp
@@ -63,7 +63,7 @@ FNiagaraSystemUserParameterBuilder::~FNiagaraSystemUserParameterBuilder()
 		{
 			if (UNiagaraUserParametersHierarchyViewModel* HierarchyViewModel = SystemViewModel.Pin()->GetUserParametersHierarchyViewModel())
 			{
-				HierarchyViewModel->GetOnHierarchyChangedDelegate().RemoveAll(this);
+				HierarchyViewModel->GetOnHierarchyEventDelegate().RemoveAll(this);
 			}
 		}
 	}
@@ -79,8 +79,8 @@ void FNiagaraSystemUserParameterBuilder::GenerateChildContent(IDetailChildrenBui
 		{
 			System->GetExposedParameters().OnStructureChanged().Add(FNiagaraParameterStore::FOnStructureChanged::FDelegate::CreateSP(this, &FNiagaraSystemUserParameterBuilder::Rebuild));
 
-			SystemViewModel.Pin()->GetUserParametersHierarchyViewModel()->GetOnHierarchyChangedDelegate().RemoveAll(this);
-			SystemViewModel.Pin()->GetUserParametersHierarchyViewModel()->GetOnHierarchyChangedDelegate().Add(UDataHierarchyViewModelBase::FOnHierarchyChanged::FDelegate::CreateSP(this, &FNiagaraSystemUserParameterBuilder::RebuildOnHierarchyChanged));
+			SystemViewModel.Pin()->GetUserParametersHierarchyViewModel()->GetOnHierarchyEventDelegate().RemoveAll(this);
+			SystemViewModel.Pin()->GetUserParametersHierarchyViewModel()->GetOnHierarchyEventDelegate().Add(UDataHierarchyViewModelBase::FOnHierarchyEvent::FDelegate::CreateSP(this, &FNiagaraSystemUserParameterBuilder::RebuildOnHierarchyChanged));
 			
 			TSharedPtr<FNiagaraUserParameterPanelViewModel> UserParameterPanelViewModel = SystemViewModel.Pin()->GetUserParameterPanelViewModel();
 			if(ensure(UserParameterPanelViewModel.IsValid()))
@@ -130,10 +130,27 @@ TSharedRef<SWidget> FNiagaraSystemUserParameterBuilder::GetAdditionalHeaderWidge
 			.Padding(2.f, 1.f)
 			[
 				SNew(SButton)
-				.OnClicked(this, &FNiagaraSystemUserParameterBuilder::SummonHierarchyEditor)
-				.Text(LOCTEXT("SummonUserParametersHierarchyButtonLabel", "Edit Hierarchy"))
-				.ButtonStyle(FAppStyle::Get(), "RoundButton")
 				.ToolTipText(LOCTEXT("SummonUserParametersHierarchyButtonTooltip", "Summon the Hierarchy Editor to add categories, sections and tooltips to User Parameters."))
+				.ButtonStyle(&FAppStyle::GetWidgetStyle<FButtonStyle>("RoundButton"))
+				.OnClicked(this, &FNiagaraSystemUserParameterBuilder::SummonHierarchyEditor)
+				[
+					SNew(SHorizontalBox)
+					+ SHorizontalBox::Slot()
+					.AutoWidth()
+					.HAlign(HAlign_Center)
+					.VAlign(VAlign_Center)
+					.Padding(4.f, 0.f)
+					[
+						SNew(SImage)
+						.Image(FAppStyle::GetBrush("Icons.EditHierarchy"))
+					]
+					+ SHorizontalBox::Slot()
+					.AutoWidth()
+					[
+						SNew(STextBlock)
+						.Text(LOCTEXT("SummonUserParametersHierarchyButtonLabel", "Edit Hierarchy"))
+					]
+				]
 			]
 			+ SHorizontalBox::Slot()
 			.HAlign(HAlign_Center)
diff --git a/Engine/Plugins/FX/Niagara/Source/NiagaraEditor/Private/Customizations/UserParameters/NiagaraUserParametersBuilderBase.h b/Engine/Plugins/FX/Niagara/Source/NiagaraEditor/Private/Customizations/UserParameters/NiagaraUserParametersBuilderBase.h
--- a/Engine/Plugins/FX/Niagara/Source/NiagaraEditor/Private/Customizations/UserParameters/NiagaraUserParametersBuilderBase.h
+++ b/Engine/Plugins/FX/Niagara/Source/NiagaraEditor/Private/Customizations/UserParameters/NiagaraUserParametersBuilderBase.h
@@ -27,7 +27,7 @@ public:
 		OnRebuildChildren.ExecuteIfBound();
 	}
 	
-	void RebuildOnHierarchyChanged(const TInstancedStruct<FHierarchyChangedPayload>& Payload) const
+	void RebuildOnHierarchyChanged(const TInstancedStruct<FHierarchyEventPayload>& Payload) const
 	{
 		Rebuild();
 	}
diff --git a/Engine/Plugins/FX/Niagara/Source/NiagaraEditor/Private/Toolkits/NiagaraScriptToolkit.cpp b/Engine/Plugins/FX/Niagara/Source/NiagaraEditor/Private/Toolkits/NiagaraScriptToolkit.cpp
--- a/Engine/Plugins/FX/Niagara/Source/NiagaraEditor/Private/Toolkits/NiagaraScriptToolkit.cpp
+++ b/Engine/Plugins/FX/Niagara/Source/NiagaraEditor/Private/Toolkits/NiagaraScriptToolkit.cpp
@@ -212,7 +212,7 @@ void FNiagaraScriptToolkit::Initialize( const EToolkitMode::Type Mode, const TSh
 
 	ParametersHierarchyViewModel = NewObject<UNiagaraScriptParametersHierarchyViewModel>();
 	ParametersHierarchyViewModel->Initialize(ScriptViewModel.ToSharedRef());
-	ParametersHierarchyViewModel->GetOnHierarchyChangedDelegate().AddSP(this, &FNiagaraScriptToolkit::OnHierarchyChanged);
+	ParametersHierarchyViewModel->GetOnHierarchyEventDelegate().AddSP(this, &FNiagaraScriptToolkit::OnHierarchyEvent);
 	
 	ParameterPanelViewModel = MakeShareable(new FNiagaraScriptToolkitParameterPanelViewModel(ScriptViewModel));
 	ParameterDefinitionsPanelViewModel = MakeShareable(new FNiagaraScriptToolkitParameterDefinitionsPanelViewModel(ScriptViewModel));
@@ -496,7 +496,7 @@ void FNiagaraScriptToolkit::OnVMScriptCompiled(UNiagaraScript*, const FGuid&)
 	}
 }
 
-void FNiagaraScriptToolkit::OnHierarchyChanged(const TInstancedStruct<FHierarchyChangedPayload>& Payload)
+void FNiagaraScriptToolkit::OnHierarchyEvent(const TInstancedStruct<FHierarchyEventPayload>& Payload)
 {
 	ScriptViewModel->GetStandaloneScript().GetScriptData()->GetSource()->MarkNotSynchronized("Input Hierarchy Changed");
 }
@@ -554,11 +554,26 @@ TSharedRef<SDockTab> FNiagaraScriptToolkit::SpawnTabScriptParameters(const FSpaw
 			.SearchAdjacentWidget()
 			[
 				SNew(SButton)
+				.ButtonStyle(&FAppStyle::GetWidgetStyle<FButtonStyle>("RoundButton"))
 				.OnClicked(this, &FNiagaraScriptToolkit::SummonHierarchyEditor)
-				.ButtonStyle(FAppStyle::Get(), "RoundButton")
 				[
-					SNew(STextBlock)
-					.Text(LOCTEXT("EditHierarchy_ScriptInputs", "Edit Input Hierarchy"))
+					SNew(SHorizontalBox)
+					+ SHorizontalBox::Slot()
+					.AutoWidth()
+					.HAlign(HAlign_Center)
+					.VAlign(VAlign_Center)
+					.Padding(4.f, 0.f)
+					[
+						SNew(SImage)
+						.Image(FAppStyle::GetBrush("Icons.EditHierarchy"))
+					]
+					+ SHorizontalBox::Slot()
+					.AutoWidth()
+					.VAlign(VAlign_Center)
+					[
+						SNew(STextBlock)
+						.Text(LOCTEXT("EditHierarchy_ScriptInputs", "Edit Hierarchy"))
+					]
 				]
 			]
 		];
diff --git a/Engine/Plugins/FX/Niagara/Source/NiagaraEditor/Private/Toolkits/NiagaraScriptToolkit.h b/Engine/Plugins/FX/Niagara/Source/NiagaraEditor/Private/Toolkits/NiagaraScriptToolkit.h
--- a/Engine/Plugins/FX/Niagara/Source/NiagaraEditor/Private/Toolkits/NiagaraScriptToolkit.h
+++ b/Engine/Plugins/FX/Niagara/Source/NiagaraEditor/Private/Toolkits/NiagaraScriptToolkit.h
@@ -96,7 +96,7 @@ private:
 
 	void OnVMScriptCompiled(UNiagaraScript* InScript, const FGuid& ScriptVersion);
 
-	void OnHierarchyChanged(const TInstancedStruct<FHierarchyChangedPayload>& Payload);
+	void OnHierarchyEvent(const TInstancedStruct<FHierarchyEventPayload>& Payload);
 	
 	/** Spawns the tab with the update graph inside */
 	TSharedRef<SDockTab> SpawnTabNodeGraph(const FSpawnTabArgs& Args);
diff --git a/Engine/Plugins/FX/Niagara/Source/NiagaraEditor/Private/Toolkits/SystemToolkitModes/NiagaraSystemToolkitModeBase.cpp b/Engine/Plugins/FX/Niagara/Source/NiagaraEditor/Private/Toolkits/SystemToolkitModes/NiagaraSystemToolkitModeBase.cpp
--- a/Engine/Plugins/FX/Niagara/Source/NiagaraEditor/Private/Toolkits/SystemToolkitModes/NiagaraSystemToolkitModeBase.cpp
+++ b/Engine/Plugins/FX/Niagara/Source/NiagaraEditor/Private/Toolkits/SystemToolkitModes/NiagaraSystemToolkitModeBase.cpp
@@ -324,22 +324,22 @@ void FNiagaraSystemToolkitModeBase::RegisterTabFactories(TSharedPtr<FTabManager>
 	InTabManager->RegisterTabSpawner(UserParametersTabID, FOnSpawnTab::CreateSP(this, &FNiagaraSystemToolkitModeBase::SpawnTab_UserParameters))
 		.SetDisplayName(LOCTEXT("UserParametersTab", "User Parameters"))
 		.SetGroup(WorkspaceMenuCategory.ToSharedRef())
-		.SetIcon(FSlateIcon(FNiagaraEditorStyle::Get().GetStyleSetName(), "Tab.UserParameters"));
+		.SetIcon(FSlateIcon(FAppStyle::GetAppStyleSetName(), "Icons.EditHierarchy"));
 	
 	InTabManager->RegisterTabSpawner(UserParametersHierarchyTabID, FOnSpawnTab::CreateSP(this, &FNiagaraSystemToolkitModeBase::SpawnTab_UserParametersHierarchyEditor))
 		.SetDisplayName(LOCTEXT("UserParametersHierarchyTab", "User Parameters Hierarchy"))
 		.SetGroup(WorkspaceMenuCategory.ToSharedRef())
-		.SetIcon(FSlateIcon(FNiagaraEditorStyle::Get().GetStyleSetName(), "Tab.UserParameterHierarchy"));
+		.SetIcon(FSlateIcon(FAppStyle::GetAppStyleSetName(), "Icons.EditHierarchy"));
 		
 	InTabManager->RegisterTabSpawner(EmitterSummaryViewEditorTabID, FOnSpawnTab::CreateSP(this, &FNiagaraSystemToolkitModeBase::SpawnTab_SummaryViewEditor))
 		.SetDisplayName(LOCTEXT("SummaryViewEditorTitle", "Edit Summary View"))
 		.SetGroup(WorkspaceMenuCategory.ToSharedRef())
-		.SetIcon(FSlateIcon(FNiagaraEditorStyle::Get().GetStyleSetName(), "Tab.UserParameterHierarchy"));
+		.SetIcon(FSlateIcon(FAppStyle::GetAppStyleSetName(), "Icons.EditHierarchy"));
 
 	InTabManager->RegisterTabSpawner(ScratchPadHierarchyEditorTabID, FOnSpawnTab::CreateSP(this, &FNiagaraSystemToolkitModeBase::SpawnTab_ScratchPadHierarchyEditor))
 		.SetDisplayName(LOCTEXT("ScratchPadHierarchyEditor", "Edit Scratch Pad Hierarchy"))
 		.SetGroup(WorkspaceMenuCategory.ToSharedRef())
-		.SetIcon(FSlateIcon(FNiagaraEditorStyle::Get().GetStyleSetName(), "Tab.UserParameterHierarchy"));
+		.SetIcon(FSlateIcon(FAppStyle::GetAppStyleSetName(), "Icons.EditHierarchy"));
 
 	RegisterTabFactoriesWithAppAndManager(PinnedSystemToolkit.Get(), InTabManager.ToSharedRef());
 }
@@ -579,12 +579,12 @@ TSharedRef<SDockTab> FNiagaraSystemToolkitModeBase::SpawnTab_SystemParameters(co
 			.SearchAdjacentWidget()
 			[
 				SNew(SButton)
-				.OnClicked(this, &FNiagaraSystemToolkitModeBase::SummonScratchPadScriptHierarchyEditor)
-				.ButtonStyle(FAppStyle::Get(), "RoundButton")
 				.Visibility(this, &FNiagaraSystemToolkitModeBase::GetSummonScratchPadHierarchyEditorButtonVisibility)
+				.ButtonStyle(&FAppStyle::GetWidgetStyle<FButtonStyle>("HoverHintOnly"))
+				.OnClicked(this, &FNiagaraSystemToolkitModeBase::SummonScratchPadScriptHierarchyEditor)
 				[
-					SNew(STextBlock)
-					.Text(LOCTEXT("EditHierarchy_ScriptInputs", "Edit Input Hierarchy"))
+					SNew(SImage)
+					.Image(FAppStyle::GetBrush("Icons.EditHierarchy"))
 				]
 			]
 		];
diff --git a/Engine/Plugins/FX/Niagara/Source/NiagaraEditor/Private/ViewModels/HierarchyEditor/NiagaraSummaryViewViewModel.cpp b/Engine/Plugins/FX/Niagara/Source/NiagaraEditor/Private/ViewModels/HierarchyEditor/NiagaraSummaryViewViewModel.cpp
--- a/Engine/Plugins/FX/Niagara/Source/NiagaraEditor/Private/ViewModels/HierarchyEditor/NiagaraSummaryViewViewModel.cpp
+++ b/Engine/Plugins/FX/Niagara/Source/NiagaraEditor/Private/ViewModels/HierarchyEditor/NiagaraSummaryViewViewModel.cpp
@@ -653,7 +653,7 @@ void FNiagaraModuleInputViewModel::AddNativeChildrenInputs()
 		GetDataMutable()->SortChildren(SortChildrenInputs, false);
 		SyncViewModelsToData();
 		
-		HierarchyViewModel->OnHierarchyChanged(TInstancedStruct<FHierarchyChangedPayload>::Make("Native children inputs added"));
+		HierarchyViewModel->OnHierarchyEvent(TInstancedStruct<FHierarchyEventPayload>::Make("Native children inputs added"));
 	}
 }
 
diff --git a/Engine/Plugins/FX/Niagara/Source/NiagaraEditor/Private/ViewModels/NiagaraEmitterViewModel.cpp b/Engine/Plugins/FX/Niagara/Source/NiagaraEditor/Private/ViewModels/NiagaraEmitterViewModel.cpp
--- a/Engine/Plugins/FX/Niagara/Source/NiagaraEditor/Private/ViewModels/NiagaraEmitterViewModel.cpp
+++ b/Engine/Plugins/FX/Niagara/Source/NiagaraEditor/Private/ViewModels/NiagaraEmitterViewModel.cpp
@@ -86,7 +86,7 @@ bool FNiagaraEmitterViewModel::Initialize(const FVersionedNiagaraEmitter& InEmit
 	SetEmitter(InEmitter);
 	SetSimulation(InSimulation);
 	SummaryViewHierarchyViewModel->Initialize(AsShared());
-	SummaryViewHierarchyViewModel->GetOnHierarchyChangedDelegate().AddSP(this, &FNiagaraEmitterViewModel::OnSummaryViewHierarchyChanged);
+	SummaryViewHierarchyViewModel->GetOnHierarchyEventDelegate().AddSP(this, &FNiagaraEmitterViewModel::OnSummaryViewHierarchyChanged);
 	GetEditorData().OnSummaryViewStateChanged().AddSP(this, &FNiagaraEmitterViewModel::UpdateSelectionOnSummaryViewStateChanged);
 	return true;
 }
@@ -94,7 +94,7 @@ bool FNiagaraEmitterViewModel::Initialize(const FVersionedNiagaraEmitter& InEmit
 void FNiagaraEmitterViewModel::Reset()
 {
 	SetEmitter(FVersionedNiagaraEmitter());
-	SummaryViewHierarchyViewModel->GetOnHierarchyChangedDelegate().RemoveAll(this);
+	SummaryViewHierarchyViewModel->GetOnHierarchyEventDelegate().RemoveAll(this);
 	SetSimulation(nullptr);
 }
 
@@ -621,7 +621,7 @@ void FNiagaraEmitterViewModel::OnEmitterPropertiesChanged()
 	}
 }
 
-void FNiagaraEmitterViewModel::OnSummaryViewHierarchyChanged(const TInstancedStruct<FHierarchyChangedPayload>& Payload)
+void FNiagaraEmitterViewModel::OnSummaryViewHierarchyChanged(const TInstancedStruct<FHierarchyEventPayload>& Payload)
 {
 	GetEditorData().OnPersistentDataChanged().Broadcast();
 }
diff --git a/Engine/Plugins/FX/Niagara/Source/NiagaraEditor/Private/ViewModels/NiagaraScratchPadScriptViewModel.cpp b/Engine/Plugins/FX/Niagara/Source/NiagaraEditor/Private/ViewModels/NiagaraScratchPadScriptViewModel.cpp
--- a/Engine/Plugins/FX/Niagara/Source/NiagaraEditor/Private/ViewModels/NiagaraScratchPadScriptViewModel.cpp
+++ b/Engine/Plugins/FX/Niagara/Source/NiagaraEditor/Private/ViewModels/NiagaraScratchPadScriptViewModel.cpp
@@ -84,7 +84,7 @@ void FNiagaraScratchPadScriptViewModel::Initialize(UNiagaraScript* Script, UNiag
 	ScriptHierarchyViewModel = NewObject<UNiagaraScriptParametersHierarchyViewModel>();
 	ScriptHierarchyViewModel->Initialize(AsShared());
 
-	ScriptHierarchyViewModel->GetOnHierarchyChangedDelegate().AddSP(this, &FNiagaraScratchPadScriptViewModel::OnHierarchyChanged);
+	ScriptHierarchyViewModel->GetOnHierarchyEventDelegate().AddSP(this, &FNiagaraScratchPadScriptViewModel::OnHierarchyEvent);
 }
 
 TSharedPtr<IAssetReferenceFilter> FNiagaraScratchPadScriptViewModel::MakeAssetReferenceFilter() const
@@ -314,7 +314,7 @@ FSimpleDelegate& FNiagaraScratchPadScriptViewModel::OnRequestDiscardChanges()
 	return OnRequestDiscardChangesDelegate;
 }
 
-void FNiagaraScratchPadScriptViewModel::OnHierarchyChanged(const TInstancedStruct<FHierarchyChangedPayload>& Payload)
+void FNiagaraScratchPadScriptViewModel::OnHierarchyEvent(const TInstancedStruct<FHierarchyEventPayload>& Payload)
 {
 	if (bHasPendingChanges == false)
 	{
diff --git a/Engine/Plugins/FX/Niagara/Source/NiagaraEditor/Private/ViewModels/Stack/NiagaraStackSummaryViewInputCollection.cpp b/Engine/Plugins/FX/Niagara/Source/NiagaraEditor/Private/ViewModels/Stack/NiagaraStackSummaryViewInputCollection.cpp
--- a/Engine/Plugins/FX/Niagara/Source/NiagaraEditor/Private/ViewModels/Stack/NiagaraStackSummaryViewInputCollection.cpp
+++ b/Engine/Plugins/FX/Niagara/Source/NiagaraEditor/Private/ViewModels/Stack/NiagaraStackSummaryViewInputCollection.cpp
@@ -41,7 +41,7 @@ void UNiagaraStackSummaryViewCollection::Initialize(FRequiredEntryData InRequire
 
 	Emitter = InEmitter;
 	GetEmitterViewModel()->GetEditorData().OnSummaryViewStateChanged().AddUObject(this, &UNiagaraStackSummaryViewCollection::OnViewStateChanged);
-	GetEmitterViewModel()->GetSummaryHierarchyViewModel()->GetOnHierarchyChangedDelegate().AddUObject(this, &UNiagaraStackSummaryViewCollection::OnHierarchyChanged);
+	GetEmitterViewModel()->GetSummaryHierarchyViewModel()->GetOnHierarchyEventDelegate().AddUObject(this, &UNiagaraStackSummaryViewCollection::OnHierarchyChanged);
 
 	AddChildFilter(FOnFilterChild::CreateUObject(this, &UNiagaraStackSummaryViewCollection::FilterOnlyModified));
 }
@@ -56,7 +56,7 @@ void UNiagaraStackSummaryViewCollection::FinalizeInternal()
 
 	if(EmitterViewModelPtr.IsValid() && EmitterViewModelPtr->GetSummaryHierarchyViewModel() != nullptr)
 	{
-		EmitterViewModelPtr->GetSummaryHierarchyViewModel()->GetOnHierarchyChangedDelegate().RemoveAll(this);
+		EmitterViewModelPtr->GetSummaryHierarchyViewModel()->GetOnHierarchyEventDelegate().RemoveAll(this);
 	}
 
 	Super::FinalizeInternal();
diff --git a/Engine/Plugins/FX/Niagara/Source/NiagaraEditor/Private/Widgets/SNiagaraScriptInputPreviewPanel.cpp b/Engine/Plugins/FX/Niagara/Source/NiagaraEditor/Private/Widgets/SNiagaraScriptInputPreviewPanel.cpp
--- a/Engine/Plugins/FX/Niagara/Source/NiagaraEditor/Private/Widgets/SNiagaraScriptInputPreviewPanel.cpp
+++ b/Engine/Plugins/FX/Niagara/Source/NiagaraEditor/Private/Widgets/SNiagaraScriptInputPreviewPanel.cpp
@@ -47,11 +47,26 @@ void SNiagaraScriptInputPreviewPanel::Construct(const FArguments& InArgs, TShare
 		.Padding(4.f)
 		[
 			SNew(SButton)
+			.ButtonStyle(&FAppStyle::GetWidgetStyle<FButtonStyle>("RoundButton"))
 			.OnClicked(this, &SNiagaraScriptInputPreviewPanel::SummonHierarchyEditor)
-			.ButtonStyle(FAppStyle::Get(), "RoundButton")
 			[
-				SNew(STextBlock)
-				.Text(LOCTEXT("EditHierarchy_ScriptInputs", "Edit Input Hierarchy"))
+				SNew(SHorizontalBox)
+				+ SHorizontalBox::Slot()
+				.AutoWidth()
+				.HAlign(HAlign_Center)
+				.VAlign(VAlign_Center)
+				.Padding(4.f, 0.f)
+				[
+					SNew(SImage)
+					.Image(FAppStyle::GetBrush("Icons.EditHierarchy"))
+				]
+				+ SHorizontalBox::Slot()
+				.AutoWidth()
+				.VAlign(VAlign_Center)
+				[
+					SNew(STextBlock)
+					.Text(LOCTEXT("EditHierarchy_ScriptInputs", "Edit Hierarchy"))
+				]
 			]
 		]
 		+ SVerticalBox::Slot()
@@ -340,7 +355,7 @@ void SNiagaraScriptInputPreviewPanel::SetupDelegates()
 {
 	Cast<UNiagaraScriptSource>(ScriptToolkit.Pin()->EditedNiagaraScript.GetScriptData()->GetSource())->NodeGraph->OnParametersChanged().AddSP(this, &SNiagaraScriptInputPreviewPanel::OnParametersChanged);
 
-	ScriptToolkit.Pin()->GetHierarchyViewModel()->GetOnHierarchyChangedDelegate().AddSP(this, &SNiagaraScriptInputPreviewPanel::RefreshOnHierarchyChanged);
+	ScriptToolkit.Pin()->GetHierarchyViewModel()->GetOnHierarchyEventDelegate().AddSP(this, &SNiagaraScriptInputPreviewPanel::RefreshOnHierarchyChanged);
 }
 
 void SNiagaraScriptInputPreviewPanel::RemoveDelegates()
@@ -351,7 +366,7 @@ void SNiagaraScriptInputPreviewPanel::RemoveDelegates()
 
 		if(UNiagaraScriptParametersHierarchyViewModel* HierarchyViewModel = ScriptToolkit.Pin()->GetHierarchyViewModel())
 		{
-			HierarchyViewModel->GetOnHierarchyChangedDelegate().RemoveAll(this);
+			HierarchyViewModel->GetOnHierarchyEventDelegate().RemoveAll(this);
 		}
 	}
 }
diff --git a/Engine/Plugins/FX/Niagara/Source/NiagaraEditor/Public/ViewModels/NiagaraEmitterViewModel.h b/Engine/Plugins/FX/Niagara/Source/NiagaraEditor/Public/ViewModels/NiagaraEmitterViewModel.h
--- a/Engine/Plugins/FX/Niagara/Source/NiagaraEditor/Public/ViewModels/NiagaraEmitterViewModel.h
+++ b/Engine/Plugins/FX/Niagara/Source/NiagaraEditor/Public/ViewModels/NiagaraEmitterViewModel.h
@@ -157,7 +157,7 @@ private:
 
 	void OnEmitterPropertiesChanged();
 
-	void OnSummaryViewHierarchyChanged(const TInstancedStruct<FHierarchyChangedPayload>& Payload);
+	void OnSummaryViewHierarchyChanged(const TInstancedStruct<FHierarchyEventPayload>& Payload);
 
 	void UpdateSelectionOnSummaryViewStateChanged() const;
 
diff --git a/Engine/Plugins/FX/Niagara/Source/NiagaraEditor/Public/ViewModels/NiagaraScratchPadScriptViewModel.h b/Engine/Plugins/FX/Niagara/Source/NiagaraEditor/Public/ViewModels/NiagaraScratchPadScriptViewModel.h
--- a/Engine/Plugins/FX/Niagara/Source/NiagaraEditor/Public/ViewModels/NiagaraScratchPadScriptViewModel.h
+++ b/Engine/Plugins/FX/Niagara/Source/NiagaraEditor/Public/ViewModels/NiagaraScratchPadScriptViewModel.h
@@ -101,7 +101,7 @@ public:
 
 	NIAGARAEDITOR_API FSimpleDelegate& OnRequestDiscardChanges();
 
-	NIAGARAEDITOR_API void OnHierarchyChanged(const TInstancedStruct<FHierarchyChangedPayload>& Payload);
+	NIAGARAEDITOR_API void OnHierarchyEvent(const TInstancedStruct<FHierarchyEventPayload>& Payload);
 
 	NIAGARAEDITOR_API FOnNodeIDFocusRequested& OnNodeIDFocusRequested();
 	NIAGARAEDITOR_API FOnPinIDFocusRequested& OnPinIDFocusRequested();
diff --git a/Engine/Plugins/FX/Niagara/Source/NiagaraEditor/Public/ViewModels/Stack/NiagaraStackSummaryViewInputCollection.h b/Engine/Plugins/FX/Niagara/Source/NiagaraEditor/Public/ViewModels/Stack/NiagaraStackSummaryViewInputCollection.h
--- a/Engine/Plugins/FX/Niagara/Source/NiagaraEditor/Public/ViewModels/Stack/NiagaraStackSummaryViewInputCollection.h
+++ b/Engine/Plugins/FX/Niagara/Source/NiagaraEditor/Public/ViewModels/Stack/NiagaraStackSummaryViewInputCollection.h
@@ -36,7 +36,7 @@ protected:
 	bool FilterOnlyModified(const UNiagaraStackEntry& Child) const;
 private:
 	NIAGARAEDITOR_API void OnViewStateChanged();
-	void OnHierarchyChanged(const TInstancedStruct<FHierarchyChangedPayload>& Payload) { OnViewStateChanged(); }
+	void OnHierarchyChanged(const TInstancedStruct<FHierarchyEventPayload>& Payload) { OnViewStateChanged(); }
 private:
 
 	FVersionedNiagaraEmitterWeakPtr Emitter;
diff --git a/Engine/Plugins/FX/Niagara/Source/NiagaraEditor/Public/Widgets/SNiagaraScriptInputPreviewPanel.h b/Engine/Plugins/FX/Niagara/Source/NiagaraEditor/Public/Widgets/SNiagaraScriptInputPreviewPanel.h
--- a/Engine/Plugins/FX/Niagara/Source/NiagaraEditor/Public/Widgets/SNiagaraScriptInputPreviewPanel.h
+++ b/Engine/Plugins/FX/Niagara/Source/NiagaraEditor/Public/Widgets/SNiagaraScriptInputPreviewPanel.h
@@ -30,7 +30,7 @@ public:
 	UE_API virtual ~SNiagaraScriptInputPreviewPanel() override;
 	
 	UE_API void Refresh();
-	void RefreshOnHierarchyChanged(const TInstancedStruct<FHierarchyChangedPayload>& Payload) { Refresh(); }
+	void RefreshOnHierarchyChanged(const TInstancedStruct<FHierarchyEventPayload>& Payload) { Refresh(); }
 	
 	UE_API void SetupDelegates();
 	UE_API void RemoveDelegates();
diff --git a/Engine/Source/Editor/DataHierarchyEditor/Private/DataHierarchyViewModelBase.cpp b/Engine/Source/Editor/DataHierarchyEditor/Private/DataHierarchyViewModelBase.cpp
--- a/Engine/Source/Editor/DataHierarchyEditor/Private/DataHierarchyViewModelBase.cpp
+++ b/Engine/Source/Editor/DataHierarchyEditor/Private/DataHierarchyViewModelBase.cpp
@@ -518,13 +518,13 @@ TSharedPtr<FHierarchyElementViewModel> UDataHierarchyViewModelBase::CreateViewMo
 	return nullptr;
 }
 
-void UDataHierarchyViewModelBase::OnHierarchyChanged(const TInstancedStruct<FHierarchyChangedPayload>& Payload)
+void UDataHierarchyViewModelBase::OnHierarchyEvent(const TInstancedStruct<FHierarchyEventPayload>& Payload)
 {
 	if(Payload.IsValid())
 	{
-		if(Payload.GetScriptStruct() == FHierarchyChangedPayload_ElementsAdded::StaticStruct())
+		if(Payload.GetScriptStruct() == FHierarchyEventPayload_ElementsAdded::StaticStruct())
 		{
-			TArray<TSharedPtr<FHierarchyElementViewModel>> AddedElementViewModels = Payload.GetPtr<FHierarchyChangedPayload_ElementsAdded>()->AddedElementViewModels;
+			TArray<TSharedPtr<FHierarchyElementViewModel>> AddedElementViewModels = Payload.GetPtr<FHierarchyEventPayload_ElementsAdded>()->AddedElementViewModels;
 			
 			bool bContainsHierarchyElement = AddedElementViewModels.ContainsByPredicate([](TSharedPtr<FHierarchyElementViewModel> Candidate)
 			{
@@ -595,9 +595,50 @@ void UDataHierarchyViewModelBase::OnHierarchyChanged(const TInstancedStruct<FHie
 			}
 		}
 
-		if(Payload.GetScriptStruct() == FHierarchyChangedPayload_ElementsDeleted::StaticStruct())
+		if (Payload.GetScriptStruct() == FHierarchyEventPayload_PreElementsDeleted::StaticStruct())
 		{
-			TArray<TSharedPtr<FHierarchyElementViewModel>> DeletedElementViewModels = Payload.GetPtr<FHierarchyChangedPayload_ElementsDeleted>()->DeletedElementViewModels;
+			TArray<TSharedPtr<FHierarchyElementViewModel>> PreDeletedElementViewModels = Payload.GetPtr<FHierarchyEventPayload_PreElementsDeleted>()->PreDeletedElementViewModels;
+			
+			for (TSharedPtr<FHierarchyElementViewModel> PreDeletedElementViewModel : PreDeletedElementViewModels)
+			{
+				if (PreDeletedElementViewModel->ShouldReparentChildrenOnDelete())
+				{
+					TArray<TSharedPtr<FHierarchyElementViewModel>> ChildrenCopy = PreDeletedElementViewModel->GetChildren();
+					for (TSharedPtr<FHierarchyElementViewModel> Child : ChildrenCopy)
+					{
+						// We determine the next valid parent for the reparenting, if it exists
+						TSharedPtr<FHierarchyElementViewModel> NextValidParent = nullptr;
+						TWeakPtr<FHierarchyElementViewModel> Candidate = PreDeletedElementViewModel->GetParent();
+						
+						// we iterate to find the next valid parent, and as long as there are candidates to evaluate 
+						while (NextValidParent == nullptr && Candidate.IsValid())
+						{
+							if (Candidate.IsValid())
+							{
+								// The candidate has to be able to contain the new data, and itself should not be flagged for deletion.
+								TSharedPtr<FHierarchyElementViewModel> CandidatePinned = Candidate.Pin();
+								if (CandidatePinned->CanContain(Child->GetData()->GetClass()).bResult && !PreDeletedElementViewModels.Contains(CandidatePinned))
+								{
+									NextValidParent = CandidatePinned;
+									break;
+								}
+								
+								Candidate = CandidatePinned->GetParent();								
+							}
+						}
+						
+						if (NextValidParent.IsValid())
+						{
+							NextValidParent->ReparentToThis(Child);
+						}
+					}
+				}
+			}
+		}
+		
+		if(Payload.GetScriptStruct() == FHierarchyEventPayload_ElementsDeleted::StaticStruct())
+		{
+			TArray<TSharedPtr<FHierarchyElementViewModel>> DeletedElementViewModels = Payload.GetPtr<FHierarchyEventPayload_ElementsDeleted>()->DeletedElementViewModels;
 			
 			bool bContainsHierarchyElement = DeletedElementViewModels.ContainsByPredicate([](TSharedPtr<FHierarchyElementViewModel> Candidate)
 			{
@@ -611,7 +652,6 @@ void UDataHierarchyViewModelBase::OnHierarchyChanged(const TInstancedStruct<FHie
 			
 			bool bContainsActiveSection = DeletedElementViewModels.Contains(ActiveHierarchySectionViewModel.Pin());
 			
-			
 			if(bContainsHierarchyElement)
 			{
 				RefreshHierarchyView(false);
@@ -628,7 +668,7 @@ void UDataHierarchyViewModelBase::OnHierarchyChanged(const TInstancedStruct<FHie
 		}
 	}
 	
-	OnHierarchyChangedDelegate.Broadcast(Payload);
+	OnHierarchyEventDelegate.Broadcast(Payload);
 }
 
 TSubclassOf<UHierarchyCategory> UDataHierarchyViewModelBase::GetCategoryDataClass() const
@@ -872,9 +912,9 @@ TSharedPtr<FHierarchyElementViewModel> FHierarchyElementViewModel::AddChild(TSub
 	TSharedPtr<FHierarchyElementViewModel> ElementViewModel = FindViewModelForChild(NewChild);
 	ensureMsgf(ElementViewModel.IsValid(), TEXT("Invalid view model detected after syncing view models to data."));
 	
-	TInstancedStruct<FHierarchyChangedPayload_ElementsAdded> Payload = TInstancedStruct<FHierarchyChangedPayload_ElementsAdded>::Make();
-	Payload.GetMutable<FHierarchyChangedPayload_ElementsAdded>().AddedElementViewModels.Add(ElementViewModel);
-	HierarchyViewModel->OnHierarchyChanged(Payload);
+	TInstancedStruct<FHierarchyEventPayload_ElementsAdded> Payload = TInstancedStruct<FHierarchyEventPayload_ElementsAdded>::Make();
+	Payload.GetMutable<FHierarchyEventPayload_ElementsAdded>().AddedElementViewModels.Add(ElementViewModel);
+	HierarchyViewModel->OnHierarchyEvent(Payload);
 	
 	return ElementViewModel;
 }
@@ -1060,9 +1100,9 @@ TSharedRef<FHierarchyElementViewModel> FHierarchyElementViewModel::DuplicateToTh
 	
 	TSharedPtr<FHierarchyElementViewModel> ViewModel = FindViewModelForChild(NewItem);
 	
-	TInstancedStruct<FHierarchyChangedPayload_ElementsAdded> Payload = TInstancedStruct<FHierarchyChangedPayload_ElementsAdded>::Make();
-	Payload.GetMutablePtr<FHierarchyChangedPayload_ElementsAdded>()->AddedElementViewModels.Add(ViewModel);
-	HierarchyViewModel->OnHierarchyChanged(Payload);
+	TInstancedStruct<FHierarchyEventPayload_ElementsAdded> Payload = TInstancedStruct<FHierarchyEventPayload_ElementsAdded>::Make();
+	Payload.GetMutablePtr<FHierarchyEventPayload_ElementsAdded>()->AddedElementViewModels.Add(ViewModel);
+	HierarchyViewModel->OnHierarchyEvent(Payload);
 	
 	return ViewModel.ToSharedRef();
 }
@@ -1085,9 +1125,9 @@ TSharedRef<FHierarchyElementViewModel> FHierarchyElementViewModel::ReparentToThi
 	
 	TSharedPtr<FHierarchyElementViewModel> ViewModel = FindViewModelForChild(NewItem);
 	
-	TInstancedStruct<FHierarchyChangedPayload_ElementsMoved> Payload = TInstancedStruct<FHierarchyChangedPayload_ElementsMoved>::Make();
-	Payload.GetMutablePtr<FHierarchyChangedPayload_ElementsMoved>()->MovedElementViewModels.Add(ViewModel);
-	HierarchyViewModel->OnHierarchyChanged(Payload);
+	TInstancedStruct<FHierarchyEventPayload_ElementsMoved> Payload = TInstancedStruct<FHierarchyEventPayload_ElementsMoved>::Make();
+	Payload.GetMutablePtr<FHierarchyEventPayload_ElementsMoved>()->MovedElementViewModels.Add(ViewModel);
+	HierarchyViewModel->OnHierarchyEvent(Payload);
 	
 	return ViewModel.ToSharedRef();
 }
@@ -1201,16 +1241,33 @@ void FHierarchyElementViewModel::Delete()
 	OnChildRequestedDeletionDelegate.Execute(AsShared());
 }
 
-void FHierarchyElementViewModel::DeleteChild(TSharedPtr<FHierarchyElementViewModel> Child)
+void FHierarchyElementViewModel::DeleteChild(TSharedPtr<FHierarchyElementViewModel> ChildToDelete)
+{
+	DeleteChildren({ChildToDelete});
+}
+
+void FHierarchyElementViewModel::DeleteChildren(TArray<TSharedPtr<FHierarchyElementViewModel>> ChildrenToDelete)
 {	
-	ensure(Child->GetParent().Pin() == AsShared());
-	GetDataMutable()->Modify();
-	GetDataMutable()->GetChildrenMutable().Remove(Child->GetDataMutable());
-	Children.Remove(Child);
+	TInstancedStruct<FHierarchyEventPayload_PreElementsDeleted> PreDeletionPayload = TInstancedStruct<FHierarchyEventPayload_PreElementsDeleted>::Make();
+	PreDeletionPayload.GetMutable<FHierarchyEventPayload_PreElementsDeleted>().PreDeletedElementViewModels.Append(ChildrenToDelete);
+	HierarchyViewModel->OnHierarchyEvent(PreDeletionPayload);
+	
+	for (TSharedPtr<FHierarchyElementViewModel> ChildToDelete : ChildrenToDelete)
+	{
+		ensure(ChildToDelete->GetParent().Pin() == AsShared());
+		GetDataMutable()->Modify();
+		GetDataMutable()->GetChildrenMutable().Remove(ChildToDelete->GetDataMutable());
+		Children.Remove(ChildToDelete);
+	}
 
-	TInstancedStruct<FHierarchyChangedPayload_ElementsDeleted> Payload = TInstancedStruct<FHierarchyChangedPayload_ElementsDeleted>::Make();
-	Payload.GetMutable<FHierarchyChangedPayload_ElementsDeleted>().DeletedElementViewModels.Add(Child);
-	HierarchyViewModel->OnHierarchyChanged(Payload);
+	TInstancedStruct<FHierarchyEventPayload_ElementsDeleted> DeletionPayload = TInstancedStruct<FHierarchyEventPayload_ElementsDeleted>::Make();
+	DeletionPayload.GetMutable<FHierarchyEventPayload_ElementsDeleted>().DeletedElementViewModels.Append(ChildrenToDelete);
+	HierarchyViewModel->OnHierarchyEvent(DeletionPayload);
+}
+
+bool FHierarchyElementViewModel::ShouldReparentChildrenOnDelete()
+{
+	return false;
 }
 
 TOptional<EItemDropZone> FHierarchyElementViewModel::OnCanRowAcceptDrop(const FDragDropEvent& DragDropEvent, EItemDropZone ItemDropZone, TSharedPtr<FHierarchyElementViewModel> Item)
@@ -1264,12 +1321,12 @@ void FHierarchyElementViewModel::Rename(FName NewName)
 	
 	RenameInternal(NewName); 
 	
-	TInstancedStruct<FHierarchyChangedPayload_ElementNameChanged> Payload = TInstancedStruct<FHierarchyChangedPayload_ElementNameChanged>::Make();
-	Payload.GetMutable<FHierarchyChangedPayload_ElementNameChanged>().ElementViewModel = AsShared();
-	Payload.GetMutable<FHierarchyChangedPayload_ElementNameChanged>().OldName = OldName;
-	Payload.GetMutable<FHierarchyChangedPayload_ElementNameChanged>().NewName = NewName;
+	TInstancedStruct<FHierarchyEventPayload_ElementNameChanged> Payload = TInstancedStruct<FHierarchyEventPayload_ElementNameChanged>::Make();
+	Payload.GetMutable<FHierarchyEventPayload_ElementNameChanged>().ElementViewModel = AsShared();
+	Payload.GetMutable<FHierarchyEventPayload_ElementNameChanged>().OldName = OldName;
+	Payload.GetMutable<FHierarchyEventPayload_ElementNameChanged>().NewName = NewName;
 	
-	HierarchyViewModel->OnHierarchyChanged(Payload);
+	HierarchyViewModel->OnHierarchyEvent(Payload);
 }
 
 bool FHierarchyElementViewModel::IsExpandedByDefault() const
@@ -1593,9 +1650,9 @@ TSharedPtr<FHierarchySectionViewModel> FHierarchyRootViewModel::AddSection()
 	
 	OnSectionAddedDelegate.ExecuteIfBound(SectionViewModel);
 	
-	TInstancedStruct<FHierarchyChangedPayload_ElementsAdded> Payload = TInstancedStruct<FHierarchyChangedPayload_ElementsAdded>::Make();
-	Payload.GetMutable<FHierarchyChangedPayload_ElementsAdded>().AddedElementViewModels.Add(SectionViewModel);
-	HierarchyViewModel->OnHierarchyChanged(Payload);
+	TInstancedStruct<FHierarchyEventPayload_ElementsAdded> Payload = TInstancedStruct<FHierarchyEventPayload_ElementsAdded>::Make();
+	Payload.GetMutable<FHierarchyEventPayload_ElementsAdded>().AddedElementViewModels.Add(SectionViewModel);
+	HierarchyViewModel->OnHierarchyEvent(Payload);
 	
 	OnSectionsChangedDelegate.ExecuteIfBound();
 	
@@ -1610,9 +1667,9 @@ void FHierarchyRootViewModel::DeleteSection(TSharedPtr<FHierarchyElementViewMode
 
 	OnSectionDeletedDelegate.ExecuteIfBound(SectionViewModel);
 	
-	TInstancedStruct<FHierarchyChangedPayload_ElementsDeleted> Payload = TInstancedStruct<FHierarchyChangedPayload_ElementsDeleted>::Make();
-	Payload.GetMutable<FHierarchyChangedPayload_ElementsDeleted>().DeletedElementViewModels.Add(SectionViewModel);
-	HierarchyViewModel->OnHierarchyChanged(Payload);
+	TInstancedStruct<FHierarchyEventPayload_ElementsDeleted> Payload = TInstancedStruct<FHierarchyEventPayload_ElementsDeleted>::Make();
+	Payload.GetMutable<FHierarchyEventPayload_ElementsDeleted>().DeletedElementViewModels.Add(SectionViewModel);
+	HierarchyViewModel->OnHierarchyEvent(Payload);
 	
 	OnSectionsChangedDelegate.ExecuteIfBound();
 }
@@ -1816,7 +1873,7 @@ TOptional<bool> FHierarchySectionViewModel::OnDroppedOnOverride(TSharedPtr<FHier
 		if(bDropSucceeded)
 		{
 			HierarchyViewModel->ForceFullRefresh();
-			HierarchyViewModel->OnHierarchyChanged(TInstancedStruct<FHierarchyChangedPayload>::Make("Dropped section next to section"));
+			HierarchyViewModel->OnHierarchyEvent(TInstancedStruct<FHierarchyEventPayload>::Make("Dropped section next to section"));
 		}
 		
 		return true;
@@ -1846,7 +1903,7 @@ TOptional<bool> FHierarchySectionViewModel::OnDroppedOnOverride(TSharedPtr<FHier
 		}
 		
 		HierarchyViewModel->RefreshHierarchyView();
-		HierarchyViewModel->OnHierarchyChanged(TInstancedStruct<FHierarchyChangedPayload>::Make("Dropped element on section"));
+		HierarchyViewModel->OnHierarchyEvent(TInstancedStruct<FHierarchyEventPayload>::Make("Dropped element on section"));
 		
 		return true;
 	}
@@ -2121,13 +2178,6 @@ void UDataHierarchyViewModelBase::DeleteElementWithIdentity(FHierarchyElementIde
 	if (ViewModelToDelete.IsValid())
 	{
 		ViewModelToDelete->Delete();
-		
-		TInstancedStruct<FHierarchyChangedPayload_ElementsDeleted> Payload = TInstancedStruct<FHierarchyChangedPayload_ElementsDeleted>::Make();
-		Payload.GetMutable<FHierarchyChangedPayload_ElementsDeleted>().DeletedElementViewModels.Add(ViewModelToDelete);
-		
-		HierarchyRootViewModel->SyncViewModelsToData();
-		
-		OnHierarchyChanged(Payload);
 	}
 	else
 	{
@@ -2148,25 +2198,25 @@ void UDataHierarchyViewModelBase::DeleteElements(TArray<TSharedPtr<FHierarchyEle
 		if(ViewModel->CanDelete())
 		{
 			ViewModelsToDelete.Add(ViewModel);
-			
 		}
 	}
 	
+	// We sort first so that the deepest elements are deleted first. This order will allow for example reparenting mechanics to work properly.
+	// With a setup like: Root -> Category1 (marked for delete) -> Category2 (marked for delete) -> Property 1:
+	// --- If Category 1 is deleted before Category 2 and reparents Category 2 to the root, the Category 2 view model will no longer for deletion.
+	// --- If Category 2 is deleted first, it will reparent Property 1 to Category 1, then Category 1 gets deleted and it gets propagated to the root.
+	ViewModelsToDelete.Sort([](TSharedPtr<FHierarchyElementViewModel> ElementA, TSharedPtr<FHierarchyElementViewModel> ElementB)
+	{
+		return ElementA->GetHierarchyDepth() > ElementB->GetHierarchyDepth();
+	});
+	
 	for (TSharedPtr<FHierarchyElementViewModel> ViewModelToDelete : ViewModelsToDelete)
 	{
 		ViewModelToDelete->Delete();
 		bAnyItemsDeleted = true;
 	}
 
-	if(bAnyItemsDeleted)
-	{
-		HierarchyRootViewModel->SyncViewModelsToData();
-		
-		TInstancedStruct<FHierarchyChangedPayload_ElementsDeleted> Payload = TInstancedStruct<FHierarchyChangedPayload_ElementsDeleted>::Make();
-		Payload.GetMutablePtr<FHierarchyChangedPayload_ElementsDeleted>()->DeletedElementViewModels = ViewModelsToDelete;
-		OnHierarchyChanged(Payload);
-	}
-	else
+	if(!bAnyItemsDeleted)
 	{
 		Transaction.Cancel();
 	}
diff --git a/Engine/Source/Editor/DataHierarchyEditor/Private/Widgets/SDataHierarchyEditor.cpp b/Engine/Source/Editor/DataHierarchyEditor/Private/Widgets/SDataHierarchyEditor.cpp
--- a/Engine/Source/Editor/DataHierarchyEditor/Private/Widgets/SDataHierarchyEditor.cpp
+++ b/Engine/Source/Editor/DataHierarchyEditor/Private/Widgets/SDataHierarchyEditor.cpp
@@ -2019,12 +2019,12 @@ void SDataHierarchyEditor::NotifyPostChange(const FPropertyChangedEvent& Propert
 	{
 		MatchingViewModel->PostEditChange(PropertyChangedEvent, PropertyThatChanged);
 		
-		TInstancedStruct<FHierarchyChangedPayload_ElementPropertyChanged> Payload = TInstancedStruct<FHierarchyChangedPayload_ElementPropertyChanged>::Make();
-		Payload.GetMutable<FHierarchyChangedPayload_ElementPropertyChanged>().ElementViewModel = MatchingViewModel;
-		Payload.GetMutable<FHierarchyChangedPayload_ElementPropertyChanged>().PropertyChangedEvent = &PropertyChangedEvent;
-		Payload.GetMutable<FHierarchyChangedPayload_ElementPropertyChanged>().PropertyThatChanged = PropertyThatChanged;
+		TInstancedStruct<FHierarchyEventPayload_ElementPropertyChanged> Payload = TInstancedStruct<FHierarchyEventPayload_ElementPropertyChanged>::Make();
+		Payload.GetMutable<FHierarchyEventPayload_ElementPropertyChanged>().ElementViewModel = MatchingViewModel;
+		Payload.GetMutable<FHierarchyEventPayload_ElementPropertyChanged>().PropertyChangedEvent = &PropertyChangedEvent;
+		Payload.GetMutable<FHierarchyEventPayload_ElementPropertyChanged>().PropertyThatChanged = PropertyThatChanged;
 
-		HierarchyViewModel->OnHierarchyChanged(Payload);
+		HierarchyViewModel->OnHierarchyEvent(Payload);
 	}
 }
 
diff --git a/Engine/Source/Editor/DataHierarchyEditor/Public/DataHierarchyViewModelBase.h b/Engine/Source/Editor/DataHierarchyEditor/Public/DataHierarchyViewModelBase.h
--- a/Engine/Source/Editor/DataHierarchyEditor/Public/DataHierarchyViewModelBase.h
+++ b/Engine/Source/Editor/DataHierarchyEditor/Public/DataHierarchyViewModelBase.h
@@ -484,21 +484,21 @@ public:
 	TArray<TSharedPtr<FHierarchyElementViewModel>> MenuHierarchyElements;
 };
 
-/** Generic payload struct for when an element in the hierarchy has changed. */
+/** Generic payload struct for when an event in the hierarchy has occurred (elements added, deleted, pre-deletion etc.). */
 USTRUCT()
-struct FHierarchyChangedPayload
+struct FHierarchyEventPayload
 {
 	GENERATED_BODY()
 	
-	FHierarchyChangedPayload() = default;
-	FHierarchyChangedPayload(FString InMessage) : Message(InMessage) {}
+	FHierarchyEventPayload() = default;
+	FHierarchyEventPayload(FString InMessage) : Message(InMessage) {}
 	
 	/** An optional message that can be set. */
 	FString Message;
 };
 
 USTRUCT()
-struct FHierarchyChangedPayload_ElementsAdded : public FHierarchyChangedPayload
+struct FHierarchyEventPayload_ElementsAdded : public FHierarchyEventPayload
 {
 	GENERATED_BODY()
 	
@@ -506,7 +506,7 @@ struct FHierarchyChangedPayload_ElementsAdded : public FHierarchyChangedPayload
 };
 
 USTRUCT()
-struct FHierarchyChangedPayload_ElementsMoved : public FHierarchyChangedPayload
+struct FHierarchyEventPayload_ElementsMoved : public FHierarchyEventPayload
 {
 	GENERATED_BODY()
 	
@@ -514,7 +514,15 @@ struct FHierarchyChangedPayload_ElementsMoved : public FHierarchyChangedPayload
 };
 
 USTRUCT()
-struct FHierarchyChangedPayload_ElementsDeleted : public FHierarchyChangedPayload
+struct FHierarchyEventPayload_PreElementsDeleted : public FHierarchyEventPayload
+{
+	GENERATED_BODY()
+
+	TArray<TSharedPtr<FHierarchyElementViewModel>> PreDeletedElementViewModels;
+};
+
+USTRUCT()
+struct FHierarchyEventPayload_ElementsDeleted : public FHierarchyEventPayload
 {
 	GENERATED_BODY()
 
@@ -522,7 +530,7 @@ struct FHierarchyChangedPayload_ElementsDeleted : public FHierarchyChangedPayloa
 };
 
 USTRUCT()
-struct FHierarchyChangedPayload_ElementPropertyChanged : public FHierarchyChangedPayload
+struct FHierarchyEventPayload_ElementPropertyChanged : public FHierarchyEventPayload
 {
 	GENERATED_BODY()
 	
@@ -532,7 +540,7 @@ struct FHierarchyChangedPayload_ElementPropertyChanged : public FHierarchyChange
 };
 
 USTRUCT()
-struct FHierarchyChangedPayload_ElementNameChanged : public FHierarchyChangedPayload
+struct FHierarchyEventPayload_ElementNameChanged : public FHierarchyEventPayload
 {
 	GENERATED_BODY()
 	
@@ -547,7 +555,7 @@ UCLASS(MinimalAPI, Abstract, Transient)
 class UDataHierarchyViewModelBase : public UObject, public FSelfRegisteringEditorUndoClient
 {
 public:
-	DECLARE_MULTICAST_DELEGATE_OneParam(FOnHierarchyChanged, const TInstancedStruct<FHierarchyChangedPayload>&)
+	DECLARE_MULTICAST_DELEGATE_OneParam(FOnHierarchyEvent, const TInstancedStruct<FHierarchyEventPayload>&)
 	DECLARE_DELEGATE_OneParam(FOnSectionActivated, TSharedPtr<FHierarchySectionViewModel> Section)
 	DECLARE_DELEGATE_OneParam(FOnRefreshViewRequested, bool bForceFullRefresh)
 	DECLARE_DELEGATE_OneParam(FOnNavigateToElementIdentityInHierarchyRequested, FHierarchyElementIdentity Identity)
@@ -587,7 +595,7 @@ public:
 	virtual void PrepareSourceItems(UHierarchyRoot* SourceRoot, TSharedPtr<FHierarchyRootViewModel> SourceRootViewModel) {};
 
 	/** Generic callback for when something in the hierarchy has changed. */
-	UE_API virtual void OnHierarchyChanged(const TInstancedStruct<FHierarchyChangedPayload>& Payload);
+	UE_API virtual void OnHierarchyEvent(const TInstancedStruct<FHierarchyEventPayload>& Payload);
 	
 	/** If specified, will add a button in the Data Hierarchy Editor that allows adding the given types from a menu rather than drag & dropping them into the hierarchy. */
 	virtual TArray<TSubclassOf<UHierarchyElement>> GetAdditionalTypesToAddInUi() const { return TArray<TSubclassOf<UHierarchyElement>>(); }
@@ -665,7 +673,7 @@ public:
 	FSimpleDelegate& OnRefreshSectionsView() { return RefreshSectionsViewDelegate; }
 
 	// Delegates for external systems
-	FOnHierarchyChanged& GetOnHierarchyChangedDelegate() { return OnHierarchyChangedDelegate; } 
+	FOnHierarchyEvent& GetOnHierarchyEventDelegate() { return OnHierarchyEventDelegate; } 
 	FOnRefreshViewRequested& OnRefreshViewRequested() { return RefreshAllViewsRequestedDelegate; }
 	FOnNavigateToElementIdentityInHierarchyRequested& OnNavigateToElementIdentityInHierarchyRequested() { return OnNavigateToElementIdentityInHierarchyRequestedDelegate; }
 	FOnNavigateToElementInHierarchyRequested& OnNavigateToElementInHierarchyRequested() { return OnNavigateToElementInHierarchyRequestedDelegate; }
@@ -724,7 +732,7 @@ protected:
 	
 	FOnSectionActivated OnHierarchySectionActivatedDelegate;
 	FOnSectionActivated OnSourceSectionActivatedDelegate;
-	FOnHierarchyChanged OnHierarchyChangedDelegate;
+	FOnHierarchyEvent OnHierarchyEventDelegate;
 	
 	FSimpleDelegate OnInitializedDelegate;
 
@@ -797,8 +805,11 @@ struct FHierarchyElementViewModel : TSharedFromThis<FHierarchyElementViewModel>,
 	
 	UE_API void AddChildFilter(FOnFilterChild InFilterChild);
 	
-	template<class DataClass, class ViewModelChildClass>
-	void GetChildrenViewModelsForType(TArray<TSharedPtr<ViewModelChildClass>>& OutChildren, bool bRecursive = false) const;
+	template<typename DataClass, typename ViewModelClass>
+	void GetChildrenViewModelsForType(TArray<TSharedPtr<ViewModelClass>>& OutChildren, bool bRecursive = false) const;
+	
+	template<typename DataClass, typename ViewModelClass, typename Predicate>
+	void GetChildrenViewModelsByPredicate(TArray<TSharedPtr<ViewModelClass>>& OutChildren, Predicate InPredicate, bool bRecursive = false) const;
 
 	/** Returns the hierarchy depth via number of parents above. */
 	UE_API int32 GetHierarchyDepth() const;
@@ -822,7 +833,11 @@ struct FHierarchyElementViewModel : TSharedFromThis<FHierarchyElementViewModel>,
 	
 	/** Deleting will ask the parent to delete its child */
 	UE_API void Delete();
-	UE_API void DeleteChild(TSharedPtr<FHierarchyElementViewModel> Child);
+	UE_API void DeleteChild(TSharedPtr<FHierarchyElementViewModel> ChildToDelete);
+	UE_API void DeleteChildren(TArray<TSharedPtr<FHierarchyElementViewModel>> ChildrenToDelete);
+	
+	/** If true, when an element gets deleted it will attempt to reparent its children up one level. If this is not valid, the child will get deleted with the element. */
+	UE_API virtual bool ShouldReparentChildrenOnDelete();
 
 	TWeakObjectPtr<UDataHierarchyViewModelBase> GetHierarchyViewModel() const { return HierarchyViewModel; }
 
@@ -997,10 +1012,10 @@ protected:
 	};
 };
 
-template <class DataClass, class ViewModelClass>
+template <typename DataClass, typename ViewModelClass>
 void FHierarchyElementViewModel::GetChildrenViewModelsForType(TArray<TSharedPtr<ViewModelClass>>& OutChildren, bool bRecursive) const
 {
-	for(auto& Child : Children)
+	for(const TSharedPtr<FHierarchyElementViewModel>& Child : Children)
 	{
 		if(Child->GetData()->IsA<DataClass>())
 		{
@@ -1010,13 +1025,36 @@ void FHierarchyElementViewModel::GetChildrenViewModelsForType(TArray<TSharedPtr<
 
 	if(bRecursive)
 	{
-		for(auto& Child : Children)
+		for(const TSharedPtr<FHierarchyElementViewModel>& Child : Children)
 		{
 			Child->GetChildrenViewModelsForType<DataClass, ViewModelClass>(OutChildren, bRecursive);
 		}
 	}
 }
 
+template <typename DataClass, typename ViewModelClass, typename Predicate>
+void FHierarchyElementViewModel::GetChildrenViewModelsByPredicate(TArray<TSharedPtr<ViewModelClass>>& OutChildren, Predicate InPredicate, bool bRecursive) const
+{
+	for(const TSharedPtr<FHierarchyElementViewModel>& Child : Children)
+	{
+		if(Child->GetData()->IsA<DataClass>())
+		{
+			if (InPredicate(Child))
+			{
+				OutChildren.Add(StaticCastSharedPtr<ViewModelClass>(Child));
+			}
+		}
+	}
+
+	if(bRecursive)
+	{
+		for(const TSharedPtr<FHierarchyElementViewModel>& Child : Children)
+		{
+			Child->GetChildrenViewModelsByPredicate<DataClass, ViewModelClass>(OutChildren, InPredicate, bRecursive);
+		}
+	}
+}
+
 struct FHierarchyRootViewModel : FHierarchyElementViewModel
 {
 	DECLARE_DELEGATE(FOnSyncPropagated)
diff --git a/Engine/Source/Editor/StructUtilsEditor/Private/HierarchyEditor/PropertyBagHierarchyViewModel.cpp b/Engine/Source/Editor/StructUtilsEditor/Private/HierarchyEditor/PropertyBagHierarchyViewModel.cpp
--- a/Engine/Source/Editor/StructUtilsEditor/Private/HierarchyEditor/PropertyBagHierarchyViewModel.cpp
+++ b/Engine/Source/Editor/StructUtilsEditor/Private/HierarchyEditor/PropertyBagHierarchyViewModel.cpp
@@ -409,6 +409,11 @@ FString FPropertyBagHierarchyPropertyViewModel::ToString() const
 	return Desc ? Desc->Name.ToString() : TEXT("Unidentified");
 }
 
+FPropertyBagHierarchyCategoryViewModel::FPropertyBagHierarchyCategoryViewModel(UPropertyBagHierarchyCategory* InCategory, TSharedRef<FHierarchyElementViewModel> InParent, TWeakObjectPtr<UDataHierarchyViewModelBase> InHierarchyViewModel)
+	: FHierarchyCategoryViewModel(InCategory, InParent, InHierarchyViewModel)
+{
+}
+
 void UPropertyBagHierarchyViewModel::Initialize(TSharedRef<IPropertyHandle> InOriginalPropertyBagHandle, TSharedRef<IPropertyUtilities> InPropertyUtilities)
 {
 	FPropertyRowGeneratorArgs Args;
@@ -515,15 +520,15 @@ UHierarchyRoot* UPropertyBagHierarchyViewModel::GetHierarchyRoot() const
 	return HierarchyRoot.Get();
 }
 
-void UPropertyBagHierarchyViewModel::OnHierarchyChanged(const TInstancedStruct<FHierarchyChangedPayload>& Payload)
+void UPropertyBagHierarchyViewModel::OnHierarchyEvent(const TInstancedStruct<FHierarchyEventPayload>& Payload)
 {
-	Super::OnHierarchyChanged(Payload);
+	Super::OnHierarchyEvent(Payload);
 	
 	if(Payload.IsValid())
 	{
-		if(Payload.GetScriptStruct() == FHierarchyChangedPayload_ElementsDeleted::StaticStruct())
+		if(Payload.GetScriptStruct() == FHierarchyEventPayload_ElementsDeleted::StaticStruct())
 		{
-			TArray<TSharedPtr<FHierarchyElementViewModel>> DeletedElementViewModel = Payload.GetPtr<FHierarchyChangedPayload_ElementsDeleted>()->DeletedElementViewModels;
+			TArray<TSharedPtr<FHierarchyElementViewModel>> DeletedElementViewModel = Payload.GetPtr<FHierarchyEventPayload_ElementsDeleted>()->DeletedElementViewModels;
 			
 			bool bContainsHierarchyElement = DeletedElementViewModel.ContainsByPredicate([](TSharedPtr<FHierarchyElementViewModel> Candidate)
 			{
@@ -568,6 +573,10 @@ TSharedPtr<FHierarchyElementViewModel> UPropertyBagHierarchyViewModel::CreateCus
 	{
 		return MakeShared<FPropertyBagHierarchyPropertyViewModel>(Property, Parent.ToSharedRef(), this);
 	}
+	else if (UPropertyBagHierarchyCategory* Category = Cast<UPropertyBagHierarchyCategory>(Element))
+	{
+		return MakeShared<FPropertyBagHierarchyCategoryViewModel>(Category, Parent.ToSharedRef(), this);
+	}
 	
 	return nullptr;
 }
diff --git a/Engine/Source/Editor/StructUtilsEditor/Public/HierarchyEditor/PropertyBagHierarchyViewModel.h b/Engine/Source/Editor/StructUtilsEditor/Public/HierarchyEditor/PropertyBagHierarchyViewModel.h
--- a/Engine/Source/Editor/StructUtilsEditor/Public/HierarchyEditor/PropertyBagHierarchyViewModel.h
+++ b/Engine/Source/Editor/StructUtilsEditor/Public/HierarchyEditor/PropertyBagHierarchyViewModel.h
@@ -179,7 +179,7 @@ struct FPropertyBagHierarchyPropertyViewModel : public FHierarchyItemViewModel
 	virtual bool DoesExternalDataStillExist(const UHierarchyDataRefreshContext* Context) const override;
 	virtual void SyncViewModelsToDataInternal() override;
 	
-	/** Users can't manually delete the properties. They can still be deleted if their parent gets deleted however.
+	/** Users can't manually delete the properties. They can still be deleted if their parent gets deleted and fails to reparent however.
 	 *  We re-add the properties to the root in that case. */
 	virtual bool CanDeleteInternal() override { return false; }
 	
@@ -187,6 +187,14 @@ protected:
 	virtual FString ToString() const override;
 };
 
+struct FPropertyBagHierarchyCategoryViewModel : public FHierarchyCategoryViewModel
+{
+	FPropertyBagHierarchyCategoryViewModel(UPropertyBagHierarchyCategory* InCategory, TSharedRef<FHierarchyElementViewModel> InParent, TWeakObjectPtr<UDataHierarchyViewModelBase> InHierarchyViewModel);
+	
+	/** We attempt to reparent our children on deletion if possible to preserve metadata. */
+	virtual bool ShouldReparentChildrenOnDelete() override { return true; }
+};
+
 /** The hierarchy view model for property bags. Defines core rules for how the hierarchy behaves. */
 UCLASS()
 class STRUCTUTILSEDITOR_API UPropertyBagHierarchyViewModel : public UDataHierarchyViewModelBase, public FTickableEditorObject
@@ -219,7 +227,7 @@ protected:
 	virtual TStatId GetStatId() const override;
 	
 	virtual UHierarchyRoot* GetHierarchyRoot() const override;
-	virtual void OnHierarchyChanged(const TInstancedStruct<FHierarchyChangedPayload>& Payload) override;
+	virtual void OnHierarchyEvent(const TInstancedStruct<FHierarchyEventPayload>& Payload) override;
 	virtual TSubclassOf<UHierarchyCategory> GetCategoryDataClass() const override;
 	virtual TSubclassOf<UHierarchySection> GetSectionDataClass() const override;
 	virtual TSharedPtr<FHierarchyElementViewModel> CreateCustomViewModelForElement(UHierarchyElement* Element, TSharedPtr<FHierarchyElementViewModel> Parent) override;
