Copy Children array before iteration because the Children array might have elements removed during the loop.

#rnx

#Synced-CL 51340398

diff --git a/Engine/Source/Editor/DataHierarchyEditor/Private/DataHierarchyViewModelBase.cpp b/Engine/Source/Editor/DataHierarchyEditor/Private/DataHierarchyViewModelBase.cpp
--- a/Engine/Source/Editor/DataHierarchyEditor/Private/DataHierarchyViewModelBase.cpp
+++ b/Engine/Source/Editor/DataHierarchyEditor/Private/DataHierarchyViewModelBase.cpp
@@ -952,7 +952,8 @@ void FHierarchyElementViewModel::RefreshChildrenData()
 	/** All remaining children are supposed to exist at this point, as internal data won't be removed by refreshing & external data was cleaned up already.
 	 * This will not call RefreshChildrenData on data that has just been added as no view models exist for these yet.
 	 */
-	for(TSharedPtr<FHierarchyElementViewModel> Child : Children)
+	TArray<TSharedPtr<FHierarchyElementViewModel>> TmpChildrenForRefresh = Children; //Copy here, as RefreshChildrenData can remove elements causing iterator invalidation
+	for(TSharedPtr<FHierarchyElementViewModel> Child : TmpChildrenForRefresh)
 	{
 		Child->RefreshChildrenData();
 	}
