- Fix null pointer crash in FHierarchyRootViewModel::CanDropOnInternal

#Synced-CL 51672242

diff --git a/Engine/Source/Editor/DataHierarchyEditor/Private/DataHierarchyViewModelBase.cpp b/Engine/Source/Editor/DataHierarchyEditor/Private/DataHierarchyViewModelBase.cpp
--- a/Engine/Source/Editor/DataHierarchyEditor/Private/DataHierarchyViewModelBase.cpp
+++ b/Engine/Source/Editor/DataHierarchyEditor/Private/DataHierarchyViewModelBase.cpp
@@ -1555,13 +1555,20 @@ FHierarchyElementViewModel::FResultWithUserFeedback FHierarchyRootViewModel::Can
 	}
 	
 	// If we are not in the All section, we only accept categories to drop directly on the root
-	if(HierarchyViewModel->GetActiveHierarchySectionViewModel()->IsAllSection() == false && !DraggedElement->GetData()->IsA<UHierarchyCategory>())
+	if (TSharedPtr<FHierarchySectionViewModel> HierarchySectionViewModel = HierarchyViewModel->GetActiveHierarchySectionViewModel())
 	{
-		Results.bResult = false;
-		FText BaseMessage = LOCTEXT("RootCanOnlyContainCategories", "While section {0} is active, the root only accepts categories. Switch to the All section first.");
-		Results.UserFeedback = FText::FormatOrdered(BaseMessage, HierarchyViewModel->GetActiveHierarchySectionViewModel()->GetSectionNameAsText());
+		if(HierarchySectionViewModel->IsAllSection() == false && !DraggedElement->GetData()->IsA<UHierarchyCategory>())
+		{
+			Results.bResult = false;
+			FText BaseMessage = LOCTEXT("RootCanOnlyContainCategories", "While section {0} is active, the root only accepts categories. Switch to the All section first.");
+			Results.UserFeedback = FText::FormatOrdered(BaseMessage, HierarchySectionViewModel->GetSectionNameAsText());
+		}
 	}
-	
+	else
+	{
+		UE_LOG(LogDataHierarchyEditor, Warning, TEXT("HierarchyViewModel (%s) has a null section view model"), *ToString());
+	}
+
 	return Results;
 }
 
