DataHierarchyEditor: Save and restore tree view expansion state across refreshes

#Synced-CL 51539883

diff --git a/Engine/Source/Editor/DataHierarchyEditor/Internal/Widgets/SDataHierarchyEditor.h b/Engine/Source/Editor/DataHierarchyEditor/Internal/Widgets/SDataHierarchyEditor.h
--- a/Engine/Source/Editor/DataHierarchyEditor/Internal/Widgets/SDataHierarchyEditor.h
+++ b/Engine/Source/Editor/DataHierarchyEditor/Internal/Widgets/SDataHierarchyEditor.h
@@ -4,6 +4,7 @@
 
 #include "CoreMinimal.h"
 #include "SDropTarget.h"
+#include "UObject/ObjectKey.h"
 #include "IDetailsView.h"
 #include "DataHierarchyEditorStyle.h"
 #include "Widgets/Views/STreeView.h"
@@ -213,6 +214,11 @@ private:
 	UE_API TOptional<SSearchBox::FSearchResultData> GetSearchResultData() const;
 
 	void ExpandEntriesByDefault(TSharedPtr<STreeView<TSharedPtr<FHierarchyElementViewModel>>> TreeView) const;
+
+	/** Saves the expansion state of all items in the given tree view, keyed by UHierarchyElement object key. */
+	void SaveExpansionState(TSharedPtr<STreeView<TSharedPtr<FHierarchyElementViewModel>>> TreeView) const;
+	/** Restores previously saved expansion state. Items not found in the saved set fall back to IsExpandedByDefault(). */
+	void RestoreExpansionState(TSharedPtr<STreeView<TSharedPtr<FHierarchyElementViewModel>>> TreeView) const;
 	
 	FHierarchyElementViewModel::FResultWithUserFeedback CanDropOnRoot(TSharedPtr<FHierarchyElementViewModel> DraggedItem) const;
 
@@ -251,6 +257,9 @@ private:
 	TSharedPtr<FActiveTimerHandle> RefreshHierarchyViewNextFrameHandle;
 	TSharedPtr<FActiveTimerHandle> RefreshSourceViewNextFrameHandle;
 	TSharedPtr<FActiveTimerHandle> RefreshSectionsViewNextFrameHandle;
+
+	/** Saved expansion state keyed by the UHierarchyElement pointer. Maps to whether the item was expanded. */
+	mutable TMap<TObjectKey<const UHierarchyElement>, bool> SavedExpansionState;
 };
 
 class SHierarchyElement : public SCompoundWidget
diff --git a/Engine/Source/Editor/DataHierarchyEditor/Private/Widgets/SDataHierarchyEditor.cpp b/Engine/Source/Editor/DataHierarchyEditor/Private/Widgets/SDataHierarchyEditor.cpp
--- a/Engine/Source/Editor/DataHierarchyEditor/Private/Widgets/SDataHierarchyEditor.cpp
+++ b/Engine/Source/Editor/DataHierarchyEditor/Private/Widgets/SDataHierarchyEditor.cpp
@@ -1020,6 +1020,9 @@ void SDataHierarchyEditor::RequestRefreshSourceViewNextFrame(bool bFullRefresh)
 
 void SDataHierarchyEditor::RefreshHierarchyView(bool bFullRefresh) const
 {
+	// Save expansion state before refreshing so we can restore it after
+	SaveExpansionState(HierarchyTreeView);
+
 	// the top layer objects might have changed due to filtering. We need to refresh these too.
 	HierarchyTreeView->SetTreeItemsSource(&HierarchyViewModel->GetHierarchyItems());
 	if(bFullRefresh)
@@ -1030,8 +1033,8 @@ void SDataHierarchyEditor::RefreshHierarchyView(bool bFullRefresh) const
 	{
 		HierarchyTreeView->RequestTreeRefresh();
 	}
-	
-	ExpandEntriesByDefault(HierarchyTreeView);
+
+	RestoreExpansionState(HierarchyTreeView);
 }
 
 void SDataHierarchyEditor::RequestRefreshHierarchyViewNextFrame(bool bFullRefresh)
@@ -1905,6 +1908,45 @@ void SDataHierarchyEditor::ExpandEntriesByDefault(TSharedPtr<STreeView<TSharedPt
 	}
 }
 
+void SDataHierarchyEditor::SaveExpansionState(TSharedPtr<STreeView<TSharedPtr<FHierarchyElementViewModel>>> TreeView) const
+{
+	SavedExpansionState.Reset();
+
+	TArray<TSharedPtr<FHierarchyElementViewModel>> EntriesToProcess(TreeView->GetRootItems());
+	while (EntriesToProcess.Num() > 0)
+	{
+		TSharedPtr<FHierarchyElementViewModel> Entry = EntriesToProcess[0];
+		EntriesToProcess.RemoveAtSwap(0);
+
+		SavedExpansionState.Add(TObjectKey<const UHierarchyElement>(Entry->GetData()), TreeView->IsItemExpanded(Entry));
+		EntriesToProcess.Append(Entry->GetFilteredChildren());
+	}
+}
+
+void SDataHierarchyEditor::RestoreExpansionState(TSharedPtr<STreeView<TSharedPtr<FHierarchyElementViewModel>>> TreeView) const
+{
+	TArray<TSharedPtr<FHierarchyElementViewModel>> EntriesToProcess(TreeView->GetRootItems());
+	while (EntriesToProcess.Num() > 0)
+	{
+		TSharedPtr<FHierarchyElementViewModel> Entry = EntriesToProcess[0];
+		EntriesToProcess.RemoveAtSwap(0);
+
+		bool bShouldExpand = Entry->IsExpandedByDefault();
+
+		if (const bool* SavedState = SavedExpansionState.Find(TObjectKey<const UHierarchyElement>(Entry->GetData())))
+		{
+			bShouldExpand = *SavedState;
+		}
+
+		TreeView->SetItemExpansion(Entry, bShouldExpand);
+
+		if (TreeView->IsItemExpanded(Entry))
+		{
+			EntriesToProcess.Append(Entry->GetFilteredChildren());
+		}
+	}
+}
+
 FReply SDataHierarchyEditor::HandleHierarchyRootDrop(const FGeometry& Geometry, const FDragDropEvent& DragDropEvent) const
 {
 	if(TSharedPtr<FHierarchyDragDropOp> DragDropOp = DragDropEvent.GetOperationAs<FHierarchyDragDropOp>())
