DataHierarchyEditor: Preserve tree view expansion state across refreshes using element identity with object key fallback

#Synced-CL 51539883

diff --git a/Engine/Source/Editor/DataHierarchyEditor/Internal/Widgets/SDataHierarchyEditor.h b/Engine/Source/Editor/DataHierarchyEditor/Internal/Widgets/SDataHierarchyEditor.h
--- a/Engine/Source/Editor/DataHierarchyEditor/Internal/Widgets/SDataHierarchyEditor.h
+++ b/Engine/Source/Editor/DataHierarchyEditor/Internal/Widgets/SDataHierarchyEditor.h
@@ -215,9 +215,9 @@ private:
 
 	void ExpandEntriesByDefault(TSharedPtr<STreeView<TSharedPtr<FHierarchyElementViewModel>>> TreeView) const;
 
-	/** Saves the expansion state of all items in the given tree view, keyed by UHierarchyElement object key. */
+	/** Saves the expansion state of all items in the given tree view, keyed by FHierarchyElementIdentity when available, falling back to TObjectKey. */
 	void SaveExpansionState(TSharedPtr<STreeView<TSharedPtr<FHierarchyElementViewModel>>> TreeView) const;
-	/** Restores previously saved expansion state. Items not found in the saved set fall back to IsExpandedByDefault(). */
+	/** Restores previously saved expansion state. Items not in the saved state fall back to IsExpandedByDefault(). */
 	void RestoreExpansionState(TSharedPtr<STreeView<TSharedPtr<FHierarchyElementViewModel>>> TreeView) const;
 	
 	FHierarchyElementViewModel::FResultWithUserFeedback CanDropOnRoot(TSharedPtr<FHierarchyElementViewModel> DraggedItem) const;
@@ -258,8 +258,9 @@ private:
 	TSharedPtr<FActiveTimerHandle> RefreshSourceViewNextFrameHandle;
 	TSharedPtr<FActiveTimerHandle> RefreshSectionsViewNextFrameHandle;
 
-	/** Saved expansion state keyed by the UHierarchyElement pointer. Maps to whether the item was expanded. */
-	mutable TMap<TObjectKey<const UHierarchyElement>, bool> SavedExpansionState;
+	/** Saved expansion state. Uses FHierarchyElementIdentity when available, falls back to TObjectKey for elements without a valid identity. */
+	mutable TMap<FHierarchyElementIdentity, bool> SavedExpansionByIdentity;
+	mutable TMap<TObjectKey<const UHierarchyElement>, bool> SavedExpansionByObjectKey;
 };
 
 class SHierarchyElement : public SCompoundWidget
diff --git a/Engine/Source/Editor/DataHierarchyEditor/Private/Widgets/SDataHierarchyEditor.cpp b/Engine/Source/Editor/DataHierarchyEditor/Private/Widgets/SDataHierarchyEditor.cpp
--- a/Engine/Source/Editor/DataHierarchyEditor/Private/Widgets/SDataHierarchyEditor.cpp
+++ b/Engine/Source/Editor/DataHierarchyEditor/Private/Widgets/SDataHierarchyEditor.cpp
@@ -1910,7 +1910,8 @@ void SDataHierarchyEditor::ExpandEntriesByDefault(TSharedPtr<STreeView<TSharedPt
 
 void SDataHierarchyEditor::SaveExpansionState(TSharedPtr<STreeView<TSharedPtr<FHierarchyElementViewModel>>> TreeView) const
 {
-	SavedExpansionState.Reset();
+	SavedExpansionByIdentity.Reset();
+	SavedExpansionByObjectKey.Reset();
 
 	TArray<TSharedPtr<FHierarchyElementViewModel>> EntriesToProcess(TreeView->GetRootItems());
 	while (EntriesToProcess.Num() > 0)
@@ -1918,7 +1919,17 @@ void SDataHierarchyEditor::SaveExpansionState(TSharedPtr<STreeView<TSharedPtr<FH
 		TSharedPtr<FHierarchyElementViewModel> Entry = EntriesToProcess[0];
 		EntriesToProcess.RemoveAtSwap(0);
 
-		SavedExpansionState.Add(TObjectKey<const UHierarchyElement>(Entry->GetData()), TreeView->IsItemExpanded(Entry));
+		const bool bIsExpanded = TreeView->IsItemExpanded(Entry);
+		const FHierarchyElementIdentity Identity = Entry->GetData()->GetPersistentIdentity();
+		if (Identity.IsValid())
+		{
+			SavedExpansionByIdentity.Add(Identity, bIsExpanded);
+		}
+		else
+		{
+			SavedExpansionByObjectKey.Add(TObjectKey<const UHierarchyElement>(Entry->GetData()), bIsExpanded);
+		}
+
 		EntriesToProcess.Append(Entry->GetFilteredChildren());
 	}
 }
@@ -1933,9 +1944,20 @@ void SDataHierarchyEditor::RestoreExpansionState(TSharedPtr<STreeView<TSharedPtr
 
 		bool bShouldExpand = Entry->IsExpandedByDefault();
 
-		if (const bool* SavedState = SavedExpansionState.Find(TObjectKey<const UHierarchyElement>(Entry->GetData())))
+		const FHierarchyElementIdentity Identity = Entry->GetData()->GetPersistentIdentity();
+		if (Identity.IsValid())
 		{
-			bShouldExpand = *SavedState;
+			if (const bool* SavedState = SavedExpansionByIdentity.Find(Identity))
+			{
+				bShouldExpand = *SavedState;
+			}
+		}
+		else
+		{
+			if (const bool* SavedState = SavedExpansionByObjectKey.Find(TObjectKey<const UHierarchyElement>(Entry->GetData())))
+			{
+				bShouldExpand = *SavedState;
+			}
 		}
 
 		TreeView->SetItemExpansion(Entry, bShouldExpand);
