[DataHierarchyEditor] Set source selection mode to Single, make hierarchy selection mode configurable via view model

#Synced-CL 52537434

diff --git a/Engine/Source/Editor/DataHierarchyEditor/Private/Widgets/SDataHierarchyEditor.cpp b/Engine/Source/Editor/DataHierarchyEditor/Private/Widgets/SDataHierarchyEditor.cpp
--- a/Engine/Source/Editor/DataHierarchyEditor/Private/Widgets/SDataHierarchyEditor.cpp
+++ b/Engine/Source/Editor/DataHierarchyEditor/Private/Widgets/SDataHierarchyEditor.cpp
@@ -741,6 +741,7 @@ void SDataHierarchyEditor::Construct(const FArguments& InArgs, TObjectPtr<UDataH
 					.Padding(1.f, 0.f)
 					[
 						SAssignNew(HierarchyTreeView, STreeView<TSharedPtr<FHierarchyElementViewModel>>)
+						.SelectionMode(InHierarchyViewModel->GetHierarchySelectionMode())
 						.TreeItemsSource(&InHierarchyViewModel->GetHierarchyItems())
 						.OnSelectionChanged(this, &SDataHierarchyEditor::OnSelectionChanged)
 						.OnGenerateRow(this, &SDataHierarchyEditor::GenerateHierarchyItemRow)
@@ -782,6 +783,7 @@ void SDataHierarchyEditor::Construct(const FArguments& InArgs, TObjectPtr<UDataH
 		.Padding(1.f, 2.f)
 		[
 			SAssignNew(SourceTreeView, STreeView<TSharedPtr<FHierarchyElementViewModel>>)
+			.SelectionMode(ESelectionMode::Single)
 			.TreeItemsSource(&GetSourceItems())
 			.OnSelectionChanged(this, &SDataHierarchyEditor::OnSelectionChanged)
 			.OnGenerateRow(this, &SDataHierarchyEditor::GenerateSourceItemRow)
diff --git a/Engine/Source/Editor/DataHierarchyEditor/Public/DataHierarchyViewModelBase.h b/Engine/Source/Editor/DataHierarchyEditor/Public/DataHierarchyViewModelBase.h
--- a/Engine/Source/Editor/DataHierarchyEditor/Public/DataHierarchyViewModelBase.h
+++ b/Engine/Source/Editor/DataHierarchyEditor/Public/DataHierarchyViewModelBase.h
@@ -608,6 +608,9 @@ public:
 	/** If false, disables source item functionality (left panel). */
 	virtual bool SupportsSourcePanel() const { return true; }
 
+	/** The selection mode used for the hierarchy tree view. Override to change the default multi-selection behavior. */
+	virtual ESelectionMode::Type GetHierarchySelectionMode() const { return ESelectionMode::Multi; }
+
 	/** Attempts to select the root object when no object is selected. Can be used in conjunction with `GetObjectForEditing' to point at external objects. */
 	virtual bool SelectRootInsteadOfNone() const { return false; }
 	
diff --git a/Engine/Source/Editor/StructUtilsEditor/Public/HierarchyEditor/PropertyBagHierarchyViewModel.h b/Engine/Source/Editor/StructUtilsEditor/Public/HierarchyEditor/PropertyBagHierarchyViewModel.h
--- a/Engine/Source/Editor/StructUtilsEditor/Public/HierarchyEditor/PropertyBagHierarchyViewModel.h
+++ b/Engine/Source/Editor/StructUtilsEditor/Public/HierarchyEditor/PropertyBagHierarchyViewModel.h
@@ -239,6 +239,7 @@ protected:
 	virtual TSharedPtr<FHierarchyElementViewModel> CreateCustomViewModelForElement(UHierarchyElement* Element, TSharedPtr<FHierarchyElementViewModel> Parent) override;
 	/** We auto-add elements to the hierarchy, so we don't need the source panel. */
 	virtual bool SupportsSourcePanel() const override { return false; }
+	virtual ESelectionMode::Type GetHierarchySelectionMode() const override { return ESelectionMode::Single; }
 	virtual void PostUndo(bool bSuccess) override;
 	virtual void PostRedo(bool bSuccess) override;
 	
